The curve editor transform tool no longer trips an ensure if views are switched while you have a selection.

diff --git a/Engine/Plugins/Editor/CurveEditorTools/Source/CurveEditorTools/Private/Tools/CurveEditorTransformTool.cpp b/Engine/Plugins/Editor/CurveEditorTools/Source/CurveEditorTools/Private/Tools/CurveEditorTransformTool.cpp
--- a/Engine/Plugins/Editor/CurveEditorTools/Source/CurveEditorTools/Private/Tools/CurveEditorTransformTool.cpp
+++ b/Engine/Plugins/Editor/CurveEditorTools/Source/CurveEditorTools/Private/Tools/CurveEditorTransformTool.cpp
@@ -172,6 +172,13 @@ void FCurveEditorTransformTool::UpdateMarqueeBoundingBox()
 			continue;
 		}
 
+		// A newly created view may have a zero-size until the next tick which is a problem if
+		// we ask the View for it's curve space, so we skip over it until it has a size.
+		if(View->GetCachedGeometry().GetLocalSize() == FVector2D::ZeroVector)
+		{
+			continue;
+		}
+
 		FCurveModel* CurveModel = CurveEditor->FindCurve(Pair.Key);
 		check(CurveModel);
 
