This is a significant overhaul to the Curve Editor used by Sequencer which adds a plugin-based architecture and extensibility. New tools and toolbar buttons can be added to all usages of the curve editor via user plugins, and the different views for data can be created modularly so new implementations of the editor can register their own way of drawing their data and the tools should just work. Additionally, you can now write your own filters to operate on curve editor data for custom implementations of smoothing, key generation, etc.

The curve editor supports three view types by default - an absolute view (default, matches old behavior), a stacked view and a normalized view. Stacked views draw each curve separately (so non-overlapping) and normalized against their own min/max values. The normalized view draws all curves overlapping with each one normalized against its own min/max values.

A tree view has been added to help effectively manage large numbers of curves. Selecting curves in the treeview controls which curves are visible in the view area. The treeview also supports pinning curves. These pinned curves will always be visible regardless of your selection in the tree view.

A transform tool and a retiming tool have been implemented (via a plugin) which is enabled by default. The transform tool allows you to do a marquee selection of keys and then translate and scale the positions of these keys. The retiming tool allows you to create a 1 dimensional lattice to adjust the timing of your keys with a linear falloff between each lattice point. These tools work across multiple views at the same time which is especially useful if you are representing one dimensional data (such as event keys) in a view, as it allows you to adjust this data at the same time as your animation curves. A smoothing filter has been implemented (via a plugin) to allow running highpass and lowpass filters on your keys.

Opening the curve editor in Sequencer/UMG now creates a separate dockable tab which can be resized and docked as desired. A time slider has been added to the Curve Editor which is synchronized to the playback time in Sequencer. This allows you to scrub time in the curve editor without having to find the Sequencer window and adjust time there while looking at your keys and previewing your animation in the viewport at the same time.

Rudimentary support has been added for saving and later restoring a set of curves in your current session. This allows you to do a rudimentary copy/paste of entire curves but can also be useful for saving a curve, making adjustments to it and then deciding you want to go back - simply reapply the saved curve! Each curve added supports an intention name (such as "Location.X" or "FieldOfView"), and these intention names will be used when trying to apply curves. This allows you to reliably take all of the curves of a transform on one object and apply them to another object (and ensure that Location.X gets applied to the new Location.X, etc.) this can be helpful if you have a mixed set of curves buffered (such as a location and a field of view). In the event that no curves match by intention you can store and apply a single curve at a time from any intention to any other intention.

The Curve Asset editors (float, vector and color curve assets) have been changed to use the new editor. They support the same treeviews, filtering and tools that the Sequencer editor does. In addition, the Color Curve asset editor adds an additional view which provides a 1 dimensional gradient editor as an easier way to visualize and edit colors instead of the channels individually.

diff --git a/Engine/Plugins/Editor/CurveEditorTools/Source/CurveEditorTools/Private/CurveEditorToolCommands.cpp b/Engine/Plugins/Editor/CurveEditorTools/Source/CurveEditorTools/Private/CurveEditorToolCommands.cpp
new file mode 100644
--- /dev/null
+++ b/Engine/Plugins/Editor/CurveEditorTools/Source/CurveEditorTools/Private/CurveEditorToolCommands.cpp
@@ -0,0 +1,18 @@
+// Copyright 1998-2018 Epic Games, Inc. All Rights Reserved.
+
+#include "CurveEditorToolCommands.h"
+
+#define LOCTEXT_NAMESPACE "CurveEditorToolCommands"
+
+void FCurveEditorToolCommands::RegisterCommands()
+{
+	// Focus Tools
+	UI_COMMAND(SetFocusPlaybackTime, "Focus Playback Time", "Focuses the Curve Editor on the current Playback Time without changing zoom level.", EUserInterfaceActionType::RadioButton, FInputChord());
+	UI_COMMAND(SetFocusPlaybackRange, "Focus Playback Range", "Focuses the Curve Editor on the current Playback Range with zoom based on visible curves.", EUserInterfaceActionType::RadioButton, FInputChord(EKeys::A));
+
+	// Tool Modes
+	UI_COMMAND(ActivateTransformTool, "Transform Tool", "Activates the Transform tool which allows translation, scale and rotation of selected keys.", EUserInterfaceActionType::ToggleButton, FInputChord(EKeys::W));
+	UI_COMMAND(ActivateRetimeTool, "Retime Tool", "Activates the Retime tool which allows you to define a one dimensional lattice to non-uniformly rescale key times.", EUserInterfaceActionType::ToggleButton, FInputChord(EKeys::E));
+}
+
+#undef LOCTEXT_NAMESPACE
diff --git a/Engine/Plugins/Editor/CurveEditorTools/Source/CurveEditorTools/Private/Tools/CurveEditorRetimeTool.h b/Engine/Plugins/Editor/CurveEditorTools/Source/CurveEditorTools/Private/Tools/CurveEditorRetimeTool.h
new file mode 100644
--- /dev/null
+++ b/Engine/Plugins/Editor/CurveEditorTools/Source/CurveEditorTools/Private/Tools/CurveEditorRetimeTool.h
@@ -0,0 +1,136 @@
+// Copyright 1998-2018 Epic Games, Inc. All Rights Reserved.
+
+#pragma once
+#include "CoreMinimal.h"
+#include "ICurveEditorToolExtension.h"
+#include "Layout/Geometry.h"
+#include "Framework/DelayedDrag.h"
+#include "Curves/KeyHandle.h"
+#include "CurveDataAbstraction.h"
+#include "CurveEditorTypes.h"
+#include "CurveEditorSelection.h"
+#include "CurveEditorRetimeTool.generated.h"
+
+class FCurveEditor;
+struct FCurveModelID;
+struct FKeyHandle;
+
+
+USTRUCT()
+struct FCurveEditorRetimeAnchor
+{
+	GENERATED_BODY()
+
+	FCurveEditorRetimeAnchor()
+		: ValueInSeconds(0.0)
+		, bIsSelected(false)
+		, bIsHighlighted(false)
+		, bIsCloseBtnHighlighted(false)
+	{
+	}
+
+	FCurveEditorRetimeAnchor(double InSeconds)
+		: ValueInSeconds(InSeconds)
+		, bIsSelected(false)
+		, bIsHighlighted(false)
+		, bIsCloseBtnHighlighted(false)
+	{
+	}
+
+	/** The time on the Timeline that this anchor is anchored at. */
+	UPROPERTY()
+	double ValueInSeconds;
+
+	/** Is this anchor currently selected? */
+	UPROPERTY()
+	bool bIsSelected;
+
+	/** Is this anchor currently highlighted? An anchor can be both selected and highlighted. */
+	bool bIsHighlighted;
+
+	/** Is the close button highlighted? */
+	bool bIsCloseBtnHighlighted;
+
+	// Calculate the geometry for this anchor based on the supplied geometry.
+	void GetGeometry(const FGeometry& InWidgetGeometry, TSharedRef<FCurveEditor> InCurveEditor, FGeometry& OutBarGeometry, FGeometry& OutCloseButtonGeometry) const;
+};
+
+UCLASS()
+class UCurveEditorRetimeToolData : public UObject
+{
+	GENERATED_BODY()
+
+public:
+	// List of anchor points. Assumes they're in order from lowest input time time to greatest.
+	UPROPERTY()
+	TArray<FCurveEditorRetimeAnchor> RetimingAnchors;
+};
+
+class FCurveEditorRetimeTool : public ICurveEditorToolExtension
+{
+public:
+	FCurveEditorRetimeTool(TWeakPtr<FCurveEditor> InCurveEditor);
+	~FCurveEditorRetimeTool();
+
+	// ICurveEditorToolExtension
+	virtual void OnPaint(const FPaintArgs& Args, const FGeometry& AllottedGeometry, const FSlateRect& MyCullingRect, FSlateWindowElementList& OutDrawElements, int32 PaintOnLayerId, const FWidgetStyle& InWidgetStyle, bool bParentEnabled) const override;
+	virtual FReply OnMouseButtonDown(TSharedRef<SWidget> OwningWidget, const FGeometry& MyGeometry, const FPointerEvent& MouseEvent) override;
+	virtual FReply OnMouseButtonDoubleClick(TSharedRef<SWidget> OwningWidget, const FGeometry& InMyGeometry, const FPointerEvent& InMouseEvent) override;
+	virtual FReply OnMouseButtonUp(TSharedRef<SWidget> OwningWidget, const FGeometry& MyGeometry, const FPointerEvent& MouseEvent) override;
+	virtual FReply OnMouseMove(TSharedRef<SWidget> OwningWidget, const FGeometry& MyGeometry, const FPointerEvent& MouseEvent) override;
+	virtual void OnFocusLost(const FFocusEvent& InFocusEvent) override;
+	virtual void OnToolActivated() override;
+	virtual void OnToolDeactivated() override;
+	virtual void BindCommands(TSharedRef<FUICommandList> CommandBindings) override;
+	// ~ICurveEditorToolExtension
+
+private:
+	virtual void DrawAnchor(const FCurveEditorRetimeAnchor& InAnchor, const FCurveEditorRetimeAnchor* OptionalNextAnchor, const FPaintArgs& Args, const FGeometry& AllottedGeometry, const FSlateRect& MyCullingRect, FSlateWindowElementList& OutDrawElements, int32 PaintOnLayerId, const FWidgetStyle& InWidgetStyle, bool bParentEnabled) const;
+	virtual void OnDragStart();
+	virtual void OnDrag(const FVector2D& InStartPosition, const FVector2D& InMousePosition);
+	virtual void OnDragEnd();
+	void StopDragIfPossible();
+
+private:
+	/** Weak pointer back to the owning curve editor. */
+	TWeakPtr<FCurveEditor> WeakCurveEditor;
+
+	/** The retiming data for this tool. This is a UObject for Undo/Redo purposes. Created and added to the root set on construction of the tool. */
+	UCurveEditorRetimeToolData* RetimeData;
+
+	/** The currently open transaction (if any) */
+	TUniquePtr<class FScopedTransaction> ActiveTransaction;
+	
+	/** Set when attempting to move a drag handle. */
+	TOptional<FDelayedDrag> DelayedDrag;
+
+	struct FPreDragChannelData
+	{
+		FPreDragChannelData(FCurveModelID InCurveModel)
+			: CurveID(InCurveModel)
+		{
+		}
+
+		/** Array of all the handles in the section at the start of the drag */
+		TArray<FKeyHandle> Handles;
+		/** Array of all the above handle's times, one per index of Handles */
+		TArray<FKeyPosition> FrameNumbers;
+		
+		FCurveModelID CurveID;
+	};
+
+	struct FPreDragCurveData
+	{
+		/** Array of all the channels in the section before it was resized */
+		TArray<FPreDragChannelData> CurveChannels;
+
+		TArray<double> RetimeAnchorStartTimes;
+
+	};
+
+	/** Allocated when the drag starts and freed when the drag is finished. Used to do calculations based off of the original timing of the data. */
+	TOptional<FPreDragCurveData> PreDragCurveData;
+
+	/** Cached user selection for keys. We clear this when the tool is opened and restore it when they leave */
+	TMap<FCurveModelID, FKeyHandleSet> CachedSelectionSet;
+};
\ No newline at end of file
diff --git a/Engine/Plugins/Editor/CurveEditorTools/Source/CurveEditorTools/Private/Tools/CurveEditorTransformTool.cpp b/Engine/Plugins/Editor/CurveEditorTools/Source/CurveEditorTools/Private/Tools/CurveEditorTransformTool.cpp
new file mode 100644
--- /dev/null
+++ b/Engine/Plugins/Editor/CurveEditorTools/Source/CurveEditorTools/Private/Tools/CurveEditorTransformTool.cpp
@@ -0,0 +1,652 @@
+// Copyright 1998-2018 Epic Games, Inc. All Rights Reserved.
+
+#include "CurveEditorTransformTool.h"
+#include "CurveEditorToolCommands.h"
+#include "CurveEditor.h"
+#include "Framework/MultiBox/MultiBoxBuilder.h"
+#include "Framework/Commands/UICommandList.h"
+#include "Framework/Commands/UIAction.h"
+#include "Rendering/DrawElements.h"
+#include "Framework/DelayedDrag.h"
+#include "SCurveEditorView.h"
+#include "ScopedTransaction.h"
+#include "CurveEditorSelection.h"
+#include "SCurveEditorPanel.h"
+#include "CurveModel.h"
+#include "CurveEditorScreenSpace.h"
+#include "Containers/ArrayView.h"
+#include "CurveDataAbstraction.h"
+#include "CurveEditorSnapMetrics.h"
+
+#define LOCTEXT_NAMESPACE "CurveEditorToolCommands"
+
+namespace CurveEditorTransformTool
+{
+	constexpr float EdgeAnchorWidth = 13.f;
+	constexpr float EdgeHighlightAlpha = 0.15f;
+}
+
+void FCurveEditorTransformWidget::GetSidebarGeometry(const FGeometry& InWidgetGeometry, FGeometry& OutLeft, FGeometry& OutRight, FGeometry& OutTop, FGeometry& OutBottom) const
+{
+	FVector2D SidebarSize = FVector2D(CurveEditorTransformTool::EdgeAnchorWidth, InWidgetGeometry.GetLocalSize().Y - CurveEditorTransformTool::EdgeAnchorWidth);
+	FVector2D SidebarSizeOffset = FVector2D(CurveEditorTransformTool::EdgeAnchorWidth / 2.f, 0.f);
+	FVector2D TopbarSize = FVector2D(InWidgetGeometry.GetLocalSize().X - CurveEditorTransformTool::EdgeAnchorWidth, CurveEditorTransformTool::EdgeAnchorWidth);
+	FVector2D TopbarSizeOffset = FVector2D(0.f, CurveEditorTransformTool::EdgeAnchorWidth / 2.f);
+
+	OutLeft = InWidgetGeometry.MakeChild(SidebarSize, FSlateLayoutTransform(FVector2D(0, CurveEditorTransformTool::EdgeAnchorWidth / 2.f) - SidebarSizeOffset));
+	OutRight = InWidgetGeometry.MakeChild(SidebarSize, FSlateLayoutTransform(FVector2D(InWidgetGeometry.GetLocalSize().X, CurveEditorTransformTool::EdgeAnchorWidth / 2.f) - SidebarSizeOffset));
+	OutTop = InWidgetGeometry.MakeChild(TopbarSize, FSlateLayoutTransform(FVector2D(CurveEditorTransformTool::EdgeAnchorWidth / 2.f, 0) - TopbarSizeOffset));
+	OutBottom = InWidgetGeometry.MakeChild(TopbarSize, FSlateLayoutTransform(FVector2D(CurveEditorTransformTool::EdgeAnchorWidth / 2.f, InWidgetGeometry.GetLocalSize().Y) - TopbarSizeOffset));
+}
+
+void FCurveEditorTransformWidget::GetCornerGeometry(const FGeometry& InWidgetGeometry, FGeometry& OutTopLeft, FGeometry& OutTopRight, FGeometry& OutBottomLeft, FGeometry& OutBottomRight) const
+{
+	FVector2D CornerSize = FVector2D(CurveEditorTransformTool::EdgeAnchorWidth, CurveEditorTransformTool::EdgeAnchorWidth);
+	FVector2D HalfSizeOffset = FVector2D(CornerSize / 2.f);
+
+	FSlateLayoutTransform TopLeftPosition = FSlateLayoutTransform(FVector2D(0, 0) - HalfSizeOffset);
+	FSlateLayoutTransform TopRightPosition = FSlateLayoutTransform(FVector2D(InWidgetGeometry.GetLocalSize().X, 0) - HalfSizeOffset);
+	FSlateLayoutTransform BottomLeftPosition = FSlateLayoutTransform(FVector2D(0, InWidgetGeometry.GetLocalSize().Y) - HalfSizeOffset);
+	FSlateLayoutTransform BottomRightPosition = FSlateLayoutTransform(FVector2D(InWidgetGeometry.GetLocalSize()) - HalfSizeOffset);
+
+	OutTopLeft = InWidgetGeometry.MakeChild(CornerSize, TopLeftPosition);
+	OutTopRight = InWidgetGeometry.MakeChild(CornerSize, TopRightPosition);
+	OutBottomLeft = InWidgetGeometry.MakeChild(CornerSize, BottomLeftPosition);
+	OutBottomRight = InWidgetGeometry.MakeChild(CornerSize, BottomRightPosition);
+}
+
+void FCurveEditorTransformWidget::GetCenterGeometry(const FGeometry& InWidgetGeometry, FGeometry& OutCenter) const
+{
+	FVector2D CenterSize = InWidgetGeometry.GetLocalSize() - FVector2D(CurveEditorTransformTool::EdgeAnchorWidth, CurveEditorTransformTool::EdgeAnchorWidth);
+	FVector2D CenterOffset = FVector2D(CurveEditorTransformTool::EdgeAnchorWidth / 2.f, CurveEditorTransformTool::EdgeAnchorWidth / 2.f);
+
+	OutCenter = InWidgetGeometry.MakeChild(CenterSize, FSlateLayoutTransform(CenterOffset));
+}
+
+ECurveEditorAnchorFlags FCurveEditorTransformWidget::GetAnchorFlagsForMousePosition(const FGeometry& InWidgetGeometry, const FVector2D& InMouseScreenPosition) const
+{
+	// We store a geometry to represent each different region, updated on Tick. We check if the mouse
+	// overlaps a region and update the selection anchors depending on which region you're hovering in.
+	ECurveEditorAnchorFlags  OutFlags = ECurveEditorAnchorFlags::None;
+
+	FGeometry LeftSidebarGeometry, RightSidebarGeometry, TopSidebarGeometry, BottomSidebarGeometry;
+	FGeometry TopLeftCornerGeometry, TopRightCornerGeometry, BottomLeftCornerGeometry, BottomRightCornerGeometry;
+
+	GetSidebarGeometry(InWidgetGeometry, LeftSidebarGeometry, RightSidebarGeometry, TopSidebarGeometry, BottomSidebarGeometry);
+	GetCornerGeometry(InWidgetGeometry, TopLeftCornerGeometry, TopRightCornerGeometry, BottomLeftCornerGeometry, BottomRightCornerGeometry);
+
+	// Deflate the supplied widget by the size of our sidebars so they don't overlap before doing the center check.
+	{
+		FGeometry CenterGeometry;
+		GetCenterGeometry(InWidgetGeometry, CenterGeometry);
+
+		if (CenterGeometry.IsUnderLocation(InMouseScreenPosition))
+		{
+			OutFlags |= ECurveEditorAnchorFlags::Center;
+		}
+	}
+
+	if (LeftSidebarGeometry.IsUnderLocation(InMouseScreenPosition))
+	{
+		OutFlags |= ECurveEditorAnchorFlags::Left;
+	}
+	if (RightSidebarGeometry.IsUnderLocation(InMouseScreenPosition))
+	{
+		OutFlags |= ECurveEditorAnchorFlags::Right;
+	}
+	if (TopSidebarGeometry.IsUnderLocation(InMouseScreenPosition))
+	{
+		OutFlags |= ECurveEditorAnchorFlags::Top;
+	}
+	if (BottomSidebarGeometry.IsUnderLocation(InMouseScreenPosition))
+	{
+		OutFlags |= ECurveEditorAnchorFlags::Bottom;
+	}
+
+	if (TopLeftCornerGeometry.IsUnderLocation(InMouseScreenPosition))
+	{
+		OutFlags |= ECurveEditorAnchorFlags::Top | ECurveEditorAnchorFlags::Left;
+	}
+	if (TopRightCornerGeometry.IsUnderLocation(InMouseScreenPosition))
+	{
+		OutFlags |= ECurveEditorAnchorFlags::Top | ECurveEditorAnchorFlags::Right;
+	}
+	if (BottomLeftCornerGeometry.IsUnderLocation(InMouseScreenPosition))
+	{
+		OutFlags |= ECurveEditorAnchorFlags::Bottom | ECurveEditorAnchorFlags::Left;
+	}
+	if (BottomRightCornerGeometry.IsUnderLocation(InMouseScreenPosition))
+	{
+		OutFlags |= ECurveEditorAnchorFlags::Bottom | ECurveEditorAnchorFlags::Right;
+	}
+
+	return OutFlags;
+}
+
+
+void FCurveEditorTransformTool::OnToolActivated()
+{
+	// UE_LOG(LogTemp, Log, TEXT("TransformTool Activated."));
+	TSharedPtr<FCurveEditor> CurveEditor = WeakCurveEditor.Pin();
+	if (CurveEditor.IsValid())
+	{
+		// CurveEditor->GetSelection().SelectionChangedDelegate += OnSelectionChanged;
+	}
+
+	// Manually fire an OnSelectionChanged event so we go look up the selection and build our bounding box
+	// OnSelectionChanged();
+}
+
+void FCurveEditorTransformTool::OnToolDeactivated()
+{
+	// UE_LOG(LogTemp, Log, TEXT("TransformTool Deactivated."));
+
+	TSharedPtr<FCurveEditor> CurveEditor = WeakCurveEditor.Pin();
+	if (CurveEditor.IsValid())
+	{
+		// CurveEditor->GetSelection().SelectionChangedDelegate -= OnSelectionChanged;
+	}
+}
+
+void FCurveEditorTransformTool::UpdateMarqueeBoundingBox()
+{
+	TSharedPtr<FCurveEditor> CurveEditor = WeakCurveEditor.Pin();
+	if (!CurveEditor.IsValid())
+	{
+		return;
+	}
+	// We need to look at all selected keys and get their positions relative to the widget view. This lets us put the bounding box around the
+	// current selection, even if it goes off-screen (which it may).
+
+	TOptional<FVector2D> MinValue;
+	TOptional<FVector2D> MaxValue;
+
+	FSlateLayoutTransform AbsoluteToContainer = CurveEditor->GetPanel()->GetViewContainerGeometry().GetAccumulatedLayoutTransform();
+
+	const TMap<FCurveModelID, FKeyHandleSet>& SelectedKeySet = CurveEditor->GetSelection().GetAll();
+	for (const TPair<FCurveModelID, FKeyHandleSet> Pair : SelectedKeySet)
+	{
+		const SCurveEditorView* View = CurveEditor->FindFirstInteractiveView(Pair.Key);
+		if (!View)
+		{
+			continue;
+		}
+
+		FCurveModel* CurveModel = CurveEditor->FindCurve(Pair.Key);
+		check(CurveModel);
+
+		TArrayView<const FKeyHandle> KeyHandles = Pair.Value.AsArray();
+
+		TArray<FKeyPosition> KeyPositions;
+		KeyPositions.SetNumUninitialized(KeyHandles.Num());
+		CurveModel->GetKeyPositions(KeyHandles, KeyPositions);
+
+		FCurveEditorScreenSpace CurveSpace = View->GetCurveSpace(Pair.Key);
+		FSlateLayoutTransform InnerToOuterTransform = Concatenate(View->GetCachedGeometry().GetAccumulatedLayoutTransform(), AbsoluteToContainer.Inverse());
+
+		for (int32 i = 0; i < KeyPositions.Num(); i++)
+		{
+			FVector2D ViewSpaceLocation = FVector2D(CurveSpace.SecondsToScreen(KeyPositions[i].InputValue), CurveSpace.ValueToScreen(KeyPositions[i].OutputValue));
+			FVector2D PanelSpaceLocation = InnerToOuterTransform.TransformPoint(ViewSpaceLocation);
+
+			if (!MinValue.IsSet())
+			{
+				MinValue = PanelSpaceLocation;
+			}
+
+			if (!MaxValue.IsSet())
+			{
+				MaxValue = PanelSpaceLocation;
+			}
+
+			MinValue = FVector2D::Min(MinValue.GetValue(), PanelSpaceLocation);
+			MaxValue = FVector2D::Max(MaxValue.GetValue(), PanelSpaceLocation);
+		}
+	}
+
+	if (MinValue.IsSet() && MaxValue.IsSet())
+	{
+		FVector2D MarqueeSize = MaxValue.GetValue() - MinValue.GetValue();
+		FVector2D Offset = FVector2D::ZeroVector;
+		
+		// Enforce a minimum size for single time/value selections.
+		if (MarqueeSize.X < 8.f)
+		{
+			MarqueeSize.X = 30;
+			Offset.X = MarqueeSize.X / 2.f;
+		}
+		if (MarqueeSize.Y < 8.f)
+		{
+			MarqueeSize.Y = 30;
+			Offset.Y = MarqueeSize.Y / 2.f;
+		}
+
+		TransformWidget.Visible = true;
+		TransformWidget.Size = MarqueeSize;
+		TransformWidget.Position = MinValue.GetValue() - Offset;
+	}
+	else
+	{
+
+		// No selection, no bounding box.
+		TransformWidget.Visible = false;
+		TransformWidget.Size = FVector2D::ZeroVector;
+		TransformWidget.Position = FVector2D::ZeroVector;
+	}
+
+	// UE_LOG(LogTemp, Log, TEXT("Size: %s Position: %s"), *TransformWidget.Size.ToString(), *TransformWidget.Position.ToString());
+}
+
+FReply FCurveEditorTransformTool::OnMouseButtonDown(TSharedRef<SWidget> OwningWidget, const FGeometry& MyGeometry, const FPointerEvent& MouseEvent)
+{
+	// UE_LOG(LogTemp, Log, TEXT("FCurveEditorTransformTool::OnMouseButtonDown"));
+	DelayedDrag.Reset();
+	if (MouseEvent.GetEffectingButton() == EKeys::LeftMouseButton)
+	{
+		FGeometry WidgetGeo = TransformWidget.MakeGeometry(MyGeometry);
+		ECurveEditorAnchorFlags HitWidgetFlags = TransformWidget.GetAnchorFlagsForMousePosition(WidgetGeo, MouseEvent.GetScreenSpacePosition());
+		if (HitWidgetFlags != ECurveEditorAnchorFlags::None)
+		{
+			// TransformWidget.SelectedAnchorFlags = HitWidgetFlags;
+			
+			// Start a Delayed Drag
+			DelayedDrag = FDelayedDrag(MyGeometry.AbsoluteToLocal(MouseEvent.GetScreenSpacePosition()), MouseEvent.GetEffectingButton());
+			return FReply::Handled();
+		}
+	}
+	return FReply::Unhandled();
+}
+
+FReply FCurveEditorTransformTool::OnMouseMove(TSharedRef<SWidget> OwningWidget, const FGeometry& MyGeometry, const FPointerEvent& MouseEvent)
+{
+	// Update the Hover State of the widget.
+	if (!DelayedDrag.IsSet())
+	{
+		FGeometry WidgetGeo = TransformWidget.MakeGeometry(MyGeometry);
+		ECurveEditorAnchorFlags HitWidgetFlags = TransformWidget.GetAnchorFlagsForMousePosition(WidgetGeo, MouseEvent.GetScreenSpacePosition());
+		TransformWidget.SelectedAnchorFlags = HitWidgetFlags;
+	}
+
+	if (DelayedDrag.IsSet())
+	{
+		FReply Reply = FReply::Handled();
+
+		if (DelayedDrag->IsDragging())
+		{
+			OnDrag(MouseEvent);
+		}
+		else if (DelayedDrag->AttemptDragStart(MouseEvent))
+		{
+			InitialMousePosition = MouseEvent.GetScreenSpacePosition();
+			OnDragStart();
+
+			// Steal the capture, as we're now the authoritative widget in charge of a mouse-drag operation
+			Reply.CaptureMouse(OwningWidget);
+		}
+
+		return Reply;
+	}
+
+	return FReply::Unhandled();
+}
+
+FReply FCurveEditorTransformTool::OnMouseButtonUp(TSharedRef<SWidget> OwningWidget, const FGeometry& MyGeometry, const FPointerEvent& MouseEvent)
+{
+	// UE_LOG(LogTemp, Log, TEXT("FCurveEditorTransformTool::OnMouseButtonUp"));
+	FReply Reply = FReply::Handled();
+	if (DelayedDrag.IsSet())
+	{
+		if (DelayedDrag->IsDragging())
+		{
+			OnDragEnd();
+
+			// Only return handled if we actually started a drag
+			Reply.ReleaseMouseCapture();
+		}
+
+		DelayedDrag.Reset();
+		return Reply;
+	}
+
+	return FReply::Unhandled();
+}
+
+void FCurveEditorTransformTool::OnFocusLost(const FFocusEvent& InFocusEvent)
+{
+	// UE_LOG(LogTemp, Log, TEXT("FCurveEditorTimeTool::OnFocusLost."));
+	// We need to end our drag if we lose Window focus to close the transaction, otherwise alt-tabbing while dragging
+	// can cause a transaction to get stuck open.
+	StopDragIfPossible();
+}
+
+
+void FCurveEditorTransformTool::OnPaint(const FPaintArgs& Args, const FGeometry& AllottedGeometry, const FSlateRect& MyCullingRect, FSlateWindowElementList& OutDrawElements, int32 PaintOnLayerId, const FWidgetStyle& InWidgetStyle, bool bParentEnabled) const
+{
+	// This Geometry represents the Marquee size, but also the offset into the window
+	FGeometry WidgetGeo = TransformWidget.MakeGeometry(AllottedGeometry);
+	DrawMarqueeWidget(TransformWidget, Args, WidgetGeo, MyCullingRect, OutDrawElements, PaintOnLayerId, InWidgetStyle, bParentEnabled);
+}
+
+void FCurveEditorTransformTool::DrawMarqueeWidget(const FCurveEditorTransformWidget& InTransformWidget, const FPaintArgs& Args, const FGeometry& AllottedGeometry, const FSlateRect& MyCullingRect, FSlateWindowElementList& OutDrawElements, int32 PaintOnLayerId, const FWidgetStyle& InWidgetStyle, bool bParentEnabled) const
+{
+	if (!InTransformWidget.Visible)
+	{
+		return;
+	}
+
+	// Draw the inner marquee dotted rectangle line and the highlight
+	{
+		FLinearColor CenterHighlightColor = (InTransformWidget.SelectedAnchorFlags == ECurveEditorAnchorFlags::Center) ? FLinearColor::White.CopyWithNewOpacity(CurveEditorTransformTool::EdgeHighlightAlpha) : FLinearColor::Transparent;
+		FGeometry CenterGeometry;
+		InTransformWidget.GetCenterGeometry(AllottedGeometry, CenterGeometry);
+
+		FSlateDrawElement::MakeBox(OutDrawElements, PaintOnLayerId, CenterGeometry.ToPaintGeometry(), FEditorStyle::GetBrush(TEXT("WhiteBrush")), ESlateDrawEffect::None, CenterHighlightColor);
+		FSlateDrawElement::MakeBox(OutDrawElements, PaintOnLayerId, AllottedGeometry.ToPaintGeometry(), FEditorStyle::GetBrush(TEXT("MarqueeSelection")));
+	}
+
+	// Draw the thing that darkens the background
+	{
+		// This is comprised of four boxes that exclude the Marquee dimensions in the middle.
+	}
+
+	// Draw inner gradient which indicates tool falloff
+
+
+	// Draw edge highlight regions on mouse hover
+	{
+		FLinearColor LeftEdgeHighlightColor =	(InTransformWidget.SelectedAnchorFlags == ECurveEditorAnchorFlags::Left)	? FLinearColor::White.CopyWithNewOpacity(CurveEditorTransformTool::EdgeHighlightAlpha) : FLinearColor::Transparent;
+		FLinearColor RightEdgeHighlightColor =	(InTransformWidget.SelectedAnchorFlags == ECurveEditorAnchorFlags::Right)	? FLinearColor::White.CopyWithNewOpacity(CurveEditorTransformTool::EdgeHighlightAlpha) : FLinearColor::Transparent;
+		FLinearColor TopEdgeHighlightColor =	(InTransformWidget.SelectedAnchorFlags == ECurveEditorAnchorFlags::Top)		? FLinearColor::White.CopyWithNewOpacity(CurveEditorTransformTool::EdgeHighlightAlpha) : FLinearColor::Transparent;
+		FLinearColor BottomEdgeHighlightColor =	(InTransformWidget.SelectedAnchorFlags == ECurveEditorAnchorFlags::Bottom)	? FLinearColor::White.CopyWithNewOpacity(CurveEditorTransformTool::EdgeHighlightAlpha) : FLinearColor::Transparent;
+
+		FGeometry LeftSidebarGeometry, RightSidebarGeometry, TopSidebarGeometry, BottomSidebarGeometry;
+		InTransformWidget.GetSidebarGeometry(AllottedGeometry, LeftSidebarGeometry, RightSidebarGeometry, TopSidebarGeometry, BottomSidebarGeometry);
+
+		// Left Edge
+		FSlateDrawElement::MakeBox(OutDrawElements, PaintOnLayerId, LeftSidebarGeometry.ToPaintGeometry(), FEditorStyle::GetBrush(TEXT("WhiteBrush")), ESlateDrawEffect::None, LeftEdgeHighlightColor);
+		// Right Edge
+		FSlateDrawElement::MakeBox(OutDrawElements, PaintOnLayerId, RightSidebarGeometry.ToPaintGeometry(), FEditorStyle::GetBrush(TEXT("WhiteBrush")), ESlateDrawEffect::None, RightEdgeHighlightColor);
+		// Top Edge
+		FSlateDrawElement::MakeBox(OutDrawElements, PaintOnLayerId, TopSidebarGeometry.ToPaintGeometry(), FEditorStyle::GetBrush(TEXT("WhiteBrush")), ESlateDrawEffect::None, TopEdgeHighlightColor);
+		// Bottom Edge
+		FSlateDrawElement::MakeBox(OutDrawElements, PaintOnLayerId, BottomSidebarGeometry.ToPaintGeometry(), FEditorStyle::GetBrush(TEXT("WhiteBrush")), ESlateDrawEffect::None, BottomEdgeHighlightColor);
+	}
+
+	// Draw the four corners + highlights
+	{
+		FGeometry TopLeftCornerGeometry, TopRightCornerGeometry, BottomLeftCornerGeometry, BottomRightCornerGeometry;
+		InTransformWidget.GetCornerGeometry(AllottedGeometry, TopLeftCornerGeometry, TopRightCornerGeometry, BottomLeftCornerGeometry, BottomRightCornerGeometry);
+
+		const bool bTopLeft =		InTransformWidget.SelectedAnchorFlags == (ECurveEditorAnchorFlags::Top		| ECurveEditorAnchorFlags::Left );
+		const bool bTopRight =		InTransformWidget.SelectedAnchorFlags == (ECurveEditorAnchorFlags::Top		| ECurveEditorAnchorFlags::Right);
+		const bool bBottomLeft =	InTransformWidget.SelectedAnchorFlags == (ECurveEditorAnchorFlags::Bottom	| ECurveEditorAnchorFlags::Left );
+		const bool bBottomRight =	InTransformWidget.SelectedAnchorFlags == (ECurveEditorAnchorFlags::Bottom	| ECurveEditorAnchorFlags::Right);
+
+		FLinearColor TopLeftHighlightColor		= bTopLeft		? FLinearColor::White.CopyWithNewOpacity(CurveEditorTransformTool::EdgeHighlightAlpha) : FLinearColor::Transparent;
+		FLinearColor TopRightHighlightColor		= bTopRight		? FLinearColor::White.CopyWithNewOpacity(CurveEditorTransformTool::EdgeHighlightAlpha) : FLinearColor::Transparent;
+		FLinearColor BottomLeftHighlightColor	= bBottomLeft	? FLinearColor::White.CopyWithNewOpacity(CurveEditorTransformTool::EdgeHighlightAlpha) : FLinearColor::Transparent;
+		FLinearColor BottomRightHighlightColor	= bBottomRight  ? FLinearColor::White.CopyWithNewOpacity(CurveEditorTransformTool::EdgeHighlightAlpha) : FLinearColor::Transparent;
+
+		// Top Left (Highlight, Corner Icon)
+		FSlateDrawElement::MakeBox(OutDrawElements, PaintOnLayerId, TopLeftCornerGeometry.ToPaintGeometry(), FEditorStyle::GetBrush(TEXT("WhiteBrush")), ESlateDrawEffect::None, TopLeftHighlightColor);
+		FSlateDrawElement::MakeBox(OutDrawElements, PaintOnLayerId, TopLeftCornerGeometry.ToPaintGeometry(), FEditorStyle::GetBrush(TEXT("MarqueeSelection")));
+		// Top Right										 
+		FSlateDrawElement::MakeBox(OutDrawElements, PaintOnLayerId, TopRightCornerGeometry.ToPaintGeometry(), FEditorStyle::GetBrush(TEXT("WhiteBrush")), ESlateDrawEffect::None, TopRightHighlightColor);
+		FSlateDrawElement::MakeBox(OutDrawElements, PaintOnLayerId, TopRightCornerGeometry.ToPaintGeometry(), FEditorStyle::GetBrush(TEXT("MarqueeSelection")));
+		// Bottom Left										 
+		FSlateDrawElement::MakeBox(OutDrawElements, PaintOnLayerId, BottomLeftCornerGeometry.ToPaintGeometry(), FEditorStyle::GetBrush(TEXT("WhiteBrush")), ESlateDrawEffect::None, BottomLeftHighlightColor);
+		FSlateDrawElement::MakeBox(OutDrawElements, PaintOnLayerId, BottomLeftCornerGeometry.ToPaintGeometry(), FEditorStyle::GetBrush(TEXT("MarqueeSelection")));
+		// Bottom Right										 
+		FSlateDrawElement::MakeBox(OutDrawElements, PaintOnLayerId, BottomRightCornerGeometry.ToPaintGeometry(), FEditorStyle::GetBrush(TEXT("WhiteBrush")), ESlateDrawEffect::None, BottomRightHighlightColor);
+		FSlateDrawElement::MakeBox(OutDrawElements, PaintOnLayerId, BottomRightCornerGeometry.ToPaintGeometry(), FEditorStyle::GetBrush(TEXT("MarqueeSelection")));
+	}
+}
+
+
+void FCurveEditorTransformTool::OnDragStart()
+{
+	ActiveTransaction = MakeUnique<FScopedTransaction>(TEXT("CurveEditorTransformTool"), LOCTEXT("CurveEditorTransformToolTransaction", "Transform Key(s)"), nullptr);
+
+	TSharedPtr<FCurveEditor> CurveEditor = WeakCurveEditor.Pin();
+	if (!CurveEditor)
+	{
+		return;
+	}
+
+	// We need to cache our key data because all of our calculations have to be relative to the starting data and not the current per-frame data.
+	KeysByCurve.Reset();
+	for (const TTuple<FCurveModelID, FKeyHandleSet>& Pair : CurveEditor->GetSelection().GetAll())
+	{
+		FCurveModelID CurveID = Pair.Key;
+		FCurveModel*  Curve = CurveEditor->FindCurve(CurveID);
+
+		if (ensureAlways(Curve))
+		{
+			Curve->Modify();
+
+			TArrayView<const FKeyHandle> Handles = Pair.Value.AsArray();
+
+			FKeyData& KeyData = KeysByCurve.Emplace_GetRef(CurveID);
+			KeyData.Handles = TArray<FKeyHandle>(Handles.GetData(), Handles.Num());
+
+			KeyData.StartKeyPositions.SetNumZeroed(KeyData.Handles.Num());
+			Curve->GetKeyPositions(KeyData.Handles, KeyData.StartKeyPositions);
+		}
+	}
+	TransformWidget.StartSize = TransformWidget.Size;
+	TransformWidget.StartPosition = TransformWidget.Position;
+}
+
+void FCurveEditorTransformTool::OnDrag(const FPointerEvent& MouseEvent)
+{
+	TSharedPtr<FCurveEditor> CurveEditor = WeakCurveEditor.Pin();
+	if (!CurveEditor)
+	{
+		return;
+	}
+
+	TArray<FKeyPosition> NewKeyPositionScratch;
+
+	FSlateLayoutTransform ContainerToAbsolute = CurveEditor->GetPanel()->GetViewContainerGeometry().GetAccumulatedLayoutTransform().Inverse();
+
+	// Dragging the center is the easy case!
+	if (TransformWidget.SelectedAnchorFlags == ECurveEditorAnchorFlags::Center)
+	{
+		FVector2D AxisLockedMousePosition = CurveEditor->GetAxisSnap().GetSnappedPosition(InitialMousePosition, MouseEvent.GetScreenSpacePosition(), MouseEvent);
+		{
+			FVector2D MouseDelta = AxisLockedMousePosition - InitialMousePosition;
+			TransformWidget.Position = TransformWidget.StartPosition + MouseDelta;
+		}
+
+		for (const FKeyData& KeyData : KeysByCurve)
+		{
+			const SCurveEditorView* View = CurveEditor->FindFirstInteractiveView(KeyData.CurveID);
+			if (!View)
+			{
+				continue;
+			}
+
+			FCurveModel* CurveModel = CurveEditor->FindCurve(KeyData.CurveID);
+			check(CurveModel);
+
+			FCurveEditorScreenSpace CurveSpace = View->GetCurveSpace(KeyData.CurveID);
+
+			double DeltaInput = (AxisLockedMousePosition.X - InitialMousePosition.X) / CurveSpace.PixelsPerInput();
+			double DeltaOutput = -(AxisLockedMousePosition.Y - InitialMousePosition.Y) / CurveSpace.PixelsPerOutput();
+
+			NewKeyPositionScratch.Reset();
+			NewKeyPositionScratch.Reserve(KeyData.StartKeyPositions.Num());
+
+			for (FKeyPosition StartPosition : KeyData.StartKeyPositions)
+			{
+				StartPosition.InputValue += DeltaInput;
+				StartPosition.OutputValue += DeltaOutput;
+
+				StartPosition.InputValue = View->IsTimeSnapEnabled() ? CurveEditor->GetSnapMetrics().SnapInputSeconds(StartPosition.InputValue) : StartPosition.InputValue;
+				StartPosition.OutputValue = View->IsValueSnapEnabled() ? CurveEditor->GetSnapMetrics().SnapOutput(StartPosition.OutputValue) : StartPosition.OutputValue;
+
+				NewKeyPositionScratch.Add(StartPosition);
+			}
+
+			CurveModel->SetKeyPositions(KeyData.Handles, NewKeyPositionScratch);
+		}
+	}
+	else if (TransformWidget.SelectedAnchorFlags != ECurveEditorAnchorFlags::None)
+	{
+		// Rescaling is where things get tricky. If a user is selecting an edge we scale on one axis at once, and if they select a corner we scale on two axis at once.
+		// If they press alt, we scale relative to the center (instead of relative to the opposite corner). Because these are all bit-flag'd together we can implement
+		// this generically by scaling coordinates relative to an arbitrary center and then masking off whether it affects the X or Y axis depending on anchor flags.
+		const bool bAffectsX = ((TransformWidget.SelectedAnchorFlags & ECurveEditorAnchorFlags::Left) != ECurveEditorAnchorFlags::None ||
+			(TransformWidget.SelectedAnchorFlags & ECurveEditorAnchorFlags::Right) != ECurveEditorAnchorFlags::None);
+
+		const bool bAffectsY = ((TransformWidget.SelectedAnchorFlags & ECurveEditorAnchorFlags::Top) != ECurveEditorAnchorFlags::None ||
+			(TransformWidget.SelectedAnchorFlags & ECurveEditorAnchorFlags::Bottom) != ECurveEditorAnchorFlags::None);
+
+		// We calculate this in [0-1] space for the widget to make the logic easier to follow.
+		FVector2D ScaleCenter = FVector2D(0.5f, 0.5f);
+		const bool bScaleFromEdge = !MouseEvent.IsAltDown();
+
+		if (bScaleFromEdge)
+		{
+			if ((TransformWidget.SelectedAnchorFlags & ECurveEditorAnchorFlags::Left) != ECurveEditorAnchorFlags::None)
+			{
+				// Anchor to the right side
+				ScaleCenter.X = 1.0f; 
+			}
+
+			if ((TransformWidget.SelectedAnchorFlags & ECurveEditorAnchorFlags::Right) != ECurveEditorAnchorFlags::None)
+			{
+				// Anchor to the left side
+				ScaleCenter.X = 0.0f;
+			}
+
+			if ((TransformWidget.SelectedAnchorFlags & ECurveEditorAnchorFlags::Top) != ECurveEditorAnchorFlags::None)
+			{
+				// Anchor to the bottom side. Slate uses top left as 0 so we flip here.
+				ScaleCenter.Y = 1.0f;
+			}
+
+			if ((TransformWidget.SelectedAnchorFlags & ECurveEditorAnchorFlags::Bottom) != ECurveEditorAnchorFlags::None)
+			{
+				// Anchor to the top side
+				ScaleCenter.Y = 0.0f;
+			}
+		}
+
+		// This is the absolute change since our KeysByCurve was initialized.
+		const FVector2D MouseDelta = MouseEvent.GetScreenSpacePosition() - InitialMousePosition;
+		// TransformWidget.Size = TransformWidget.StartSize += MouseDelta;
+
+		// We have to flip the delta depending on which edge you grabbed so that the change always goes towards the mouse.
+		const float InputMulSign = (TransformWidget.SelectedAnchorFlags & ECurveEditorAnchorFlags::Left) != ECurveEditorAnchorFlags::None ? -1.0f : 1.0f;
+		const float OutputMulSign = (TransformWidget.SelectedAnchorFlags & ECurveEditorAnchorFlags::Top) != ECurveEditorAnchorFlags::None ? -1.0f : 1.0f;
+		const FVector2D AxisFixedMouseDelta = FVector2D(MouseDelta.X * InputMulSign, MouseDelta.Y * OutputMulSign);
+		const FVector2D PanelSpaceCenter = TransformWidget.StartPosition + (TransformWidget.StartSize * ScaleCenter);
+		const FVector2D PercentChanged = FVector2D(1.0f, 1.0f) + (AxisFixedMouseDelta / TransformWidget.StartSize); // ie: 5 pixel change on a 100 wide gives you 1.05
+
+		// We now know if we need to affect both X and Y, and we know where we're scaling from. Now we can loop through the keys and actually modify their positions.
+		// We perform the scale on both axis (for simplicity) and then read which axis it should effect before assigning it back to the key position.
+		for (const FKeyData& KeyData : KeysByCurve)
+		{
+			const SCurveEditorView* View = CurveEditor->FindFirstInteractiveView(KeyData.CurveID);
+			if (!View)
+			{
+				continue;
+			}
+
+			FCurveModel* CurveModel = CurveEditor->FindCurve(KeyData.CurveID);
+			check(CurveModel);
+
+			// Compute the curve-space center for transformation by transforming the center from panel space to view space, then to curve space
+			FSlateLayoutTransform OuterToInnerTransform = Concatenate(ContainerToAbsolute, View->GetCachedGeometry().GetAccumulatedLayoutTransform()).Inverse();
+			FVector2D ViewSpaceCenter = OuterToInnerTransform.TransformPoint(PanelSpaceCenter);
+
+			FCurveEditorScreenSpace CurveSpace = View->GetCurveSpace(KeyData.CurveID);
+			const double CurveSpaceCenterInput =  CurveSpace.ScreenToSeconds(ViewSpaceCenter.X);
+			const double CurveSpaceCenterOutput = CurveSpace.ScreenToValue(ViewSpaceCenter.Y);
+
+			// UE_LOG(LogTemp, Log, TEXT("ScaleCenter Input: %f ScaleCenter Output: %f PercentChanged: %s"), CenterInput, CenterOutput, *PercentChanged.ToString());
+
+			NewKeyPositionScratch.Reset();
+			NewKeyPositionScratch.Reserve(KeyData.StartKeyPositions.Num());
+
+			// UE_LOG(LogTemp, Log, TEXT("MouseDelta: %s PercentChanged: %s CenterInput: %f CenterOutput: %f"),
+			// *MouseDelta.ToString(), *PercentChanged.ToString(), CenterInput, CenterOutput);
+
+			int32 i = 0;
+			for (FKeyPosition StartPosition : KeyData.StartKeyPositions)
+			{
+				// Step 1 is to rescale the position of the key by the percentage change on each axis.
+				double ScaledInput = (StartPosition.InputValue - CurveSpaceCenterInput) * PercentChanged.X; // *InputMulSign;
+				double ScaledOutput = (StartPosition.OutputValue - CurveSpaceCenterOutput) * PercentChanged.Y; // *OutputMulSign;
+
+				// Step 2 is to subtract it from the center position so we support scaling from places other than zero
+				double NewInput = CurveSpaceCenterInput + ScaledInput;
+				double NewOutput = CurveSpaceCenterOutput + ScaledOutput;
+
+				// UE_LOG(LogTemp, Log, TEXT("[%d] InputValue: %f OutputValue: %f ScaledInput: %f ScaledOutput: %f NewInput: %f NewOutput: %f"),
+				// i, StartPosition.InputValue, StartPosition.OutputValue, ScaledInput, ScaledOutput, NewInput, NewOutput);
+				i++;
+
+				
+				// Snap the new values to the grid. We calculate both X and Y changes for ease of programming above and just limit which one it applies to.
+				// This includes snapping, otherwise dragging on an edge can cause it to snap on the opposite axis.
+				if (bAffectsX)
+				{
+					StartPosition.InputValue = View->IsTimeSnapEnabled() ? CurveEditor->GetSnapMetrics().SnapInputSeconds(NewInput) : NewInput;
+				}
+				if (bAffectsY)
+				{
+					StartPosition.OutputValue = View->IsValueSnapEnabled() ? CurveEditor->GetSnapMetrics().SnapOutput(NewOutput) : NewOutput;
+				}
+
+
+				NewKeyPositionScratch.Add(StartPosition);
+			}
+
+			CurveModel->SetKeyPositions(KeyData.Handles, NewKeyPositionScratch);
+		}
+	}
+}
+
+void FCurveEditorTransformTool::OnDragEnd()
+{
+	// This finalizes the transaction
+	ActiveTransaction.Reset();
+	// PreDragCurveData.Reset();
+	// TotalDeltaTime = 0.0;
+}
+
+void FCurveEditorTransformTool::StopDragIfPossible()
+{
+	if (DelayedDrag.IsSet())
+	{
+		if (DelayedDrag->IsDragging())
+		{
+			OnDragEnd();
+		}
+
+		DelayedDrag.Reset();
+	}
+}
+
+void FCurveEditorTransformTool::Tick(const FGeometry& AllottedGeometry, const double InCurrentTime, const float InDeltaTime)
+{
+	// We update the size and position of the box every frame as some scale operations aren't 1:1
+	// so this keeps the box visually containing all keys even if the mouse position no longer quite 
+	// matches up.
+	UpdateMarqueeBoundingBox();
+}
+
+void FCurveEditorTransformTool::BindCommands(TSharedRef<FUICommandList> CommandBindings)
+{
+	TSharedPtr<FCurveEditor> CurveEditor = WeakCurveEditor.Pin();
+	if (CurveEditor)
+	{
+		FIsActionChecked TransformToolIsActive = FIsActionChecked::CreateSP(CurveEditor.ToSharedRef(), &FCurveEditor::IsToolActive, ToolID);
+		FExecuteAction ActivateTransformTool = FExecuteAction::CreateSP(CurveEditor.ToSharedRef(), &FCurveEditor::MakeToolActive, ToolID);
+
+		CommandBindings->MapAction(FCurveEditorToolCommands::Get().ActivateTransformTool, ActivateTransformTool, FCanExecuteAction(), TransformToolIsActive);
+	}
+}
+#undef LOCTEXT_NAMESPACE // "CurveEditorToolCommands"
\ No newline at end of file
diff --git a/Engine/Plugins/MovieScene/SequencerScripting/Source/SequencerScripting/Private/ExtensionLibraries/MovieSceneSectionExtensions.cpp b/Engine/Plugins/MovieScene/SequencerScripting/Source/SequencerScripting/Private/ExtensionLibraries/MovieSceneSectionExtensions.cpp
--- a/Engine/Plugins/MovieScene/SequencerScripting/Source/SequencerScripting/Private/ExtensionLibraries/MovieSceneSectionExtensions.cpp
+++ b/Engine/Plugins/MovieScene/SequencerScripting/Source/SequencerScripting/Private/ExtensionLibraries/MovieSceneSectionExtensions.cpp
@@ -239,7 +239,7 @@ TArray<UMovieSceneScriptingChannel*> UMovieSceneSectionExtensions::GetChannels(U
 	GetScriptingChannelsForChannel<FMovieSceneFloatChannel, UMovieSceneScriptingFloatChannel>(ChannelProxy, Sequence, Channels);
 	GetScriptingChannelsForChannel<FMovieSceneIntegerChannel, UMovieSceneScriptingIntegerChannel>(ChannelProxy, Sequence, Channels);
 	GetScriptingChannelsForChannel<FMovieSceneStringChannel, UMovieSceneScriptingStringChannel>(ChannelProxy, Sequence, Channels);
-	GetScriptingChannelsForChannel<FMovieSceneEventSectionData, UMovieSceneScriptingEventChannel>(ChannelProxy, Sequence, Channels);
+	GetScriptingChannelsForChannel<FMovieSceneEventChannel, UMovieSceneScriptingEventChannel>(ChannelProxy, Sequence, Channels);
 	GetScriptingChannelsForChannel<FMovieSceneActorReferenceData, UMovieSceneScriptingActorReferenceChannel>(ChannelProxy, Sequence, Channels);
 
 	return Channels;
@@ -264,7 +264,7 @@ TArray<UMovieSceneScriptingChannel*> UMovieSceneSectionExtensions::FindChannelsB
 	else if (ChannelType == UMovieSceneScriptingFloatChannel::StaticClass())			{ GetScriptingChannelsForChannel<FMovieSceneFloatChannel, UMovieSceneScriptingFloatChannel>(ChannelProxy, Sequence, Channels); }
 	else if (ChannelType == UMovieSceneScriptingIntegerChannel::StaticClass())			{ GetScriptingChannelsForChannel<FMovieSceneIntegerChannel, UMovieSceneScriptingIntegerChannel>(ChannelProxy, Sequence, Channels); }
 	else if (ChannelType == UMovieSceneScriptingStringChannel::StaticClass())			{ GetScriptingChannelsForChannel<FMovieSceneStringChannel, UMovieSceneScriptingStringChannel>(ChannelProxy, Sequence, Channels); }
-	else if (ChannelType == UMovieSceneScriptingEventChannel::StaticClass())			{ GetScriptingChannelsForChannel<FMovieSceneEventSectionData, UMovieSceneScriptingEventChannel>(ChannelProxy, Sequence, Channels); }
+	else if (ChannelType == UMovieSceneScriptingEventChannel::StaticClass())			{ GetScriptingChannelsForChannel<FMovieSceneEventChannel, UMovieSceneScriptingEventChannel>(ChannelProxy, Sequence, Channels); }
 	else if (ChannelType == UMovieSceneScriptingActorReferenceChannel::StaticClass())	{ GetScriptingChannelsForChannel<FMovieSceneActorReferenceData, UMovieSceneScriptingActorReferenceChannel>(ChannelProxy, Sequence, Channels); }
 	else
 	{
diff --git a/Engine/Plugins/MovieScene/SequencerScripting/Source/SequencerScripting/Private/KeysAndChannels/MovieSceneScriptingEvent.h b/Engine/Plugins/MovieScene/SequencerScripting/Source/SequencerScripting/Private/KeysAndChannels/MovieSceneScriptingEvent.h
--- a/Engine/Plugins/MovieScene/SequencerScripting/Source/SequencerScripting/Private/KeysAndChannels/MovieSceneScriptingEvent.h
+++ b/Engine/Plugins/MovieScene/SequencerScripting/Source/SequencerScripting/Private/KeysAndChannels/MovieSceneScriptingEvent.h
@@ -8,9 +8,10 @@
 #include "Channels/MovieSceneChannelTraits.h"
 #include "Channels/MovieSceneChannelHandle.h"
 #include "Sections/MovieSceneEventSection.h"
-#include "KeysAndChannels/MovieSceneScriptingChannel.h"
 #include "KeyParams.h"
 #include "MovieScene.h"
+#include "Channels/MovieSceneEvent.h"
+#include "Channels/MovieSceneEventChannel.h"
 
 #include "MovieSceneScriptingEvent.generated.h"
 
@@ -19,7 +20,7 @@
 * Stores a reference to the data so changes to this class are forwarded onto the underlying data structures.
 */
 UCLASS(BlueprintType)
-class UMovieSceneScriptingEventKey : public UMovieSceneScriptingKey, public TMovieSceneScriptingKey<FMovieSceneEventSectionData, FEventPayload>
+class UMovieSceneScriptingEventKey : public UMovieSceneScriptingKey, public TMovieSceneScriptingKey<FMovieSceneEventChannel, FMovieSceneEvent>
 {
 	GENERATED_BODY()
 public:
@@ -45,7 +46,7 @@ public:
 	* @return	The event payload for this key which contains an event name and data.
 	*/
 	UFUNCTION(BlueprintCallable, Category = "Editor Scripting | Sequencer Tools | Keys", meta = (DisplayName = "Get Value (Event)"))
-	FEventPayload GetValue() const
+	FMovieSceneEvent GetValue() const
 	{
 		return GetValueFromChannel(KeyHandle);
 	}
@@ -55,14 +56,14 @@ public:
 	* @param InNewValue	The new event payload for this key.
 	*/
 	UFUNCTION(BlueprintCallable, Category = "Editor Scripting | Sequencer Tools | Keys", meta = (DisplayName = "Get Value (Event)"))
-	void SetValue(const FEventPayload& InNewValue)
+	void SetValue(const FMovieSceneEvent& InNewValue)
 	{
 		SetValueInChannel(KeyHandle, InNewValue);
 	}
 };
 
 UCLASS(BlueprintType)
-class UMovieSceneScriptingEventChannel : public UMovieSceneScriptingChannel, public TMovieSceneScriptingChannel<FMovieSceneEventSectionData, UMovieSceneScriptingEventKey, FEventPayload>
+class UMovieSceneScriptingEventChannel : public UMovieSceneScriptingChannel, public TMovieSceneScriptingChannel<FMovieSceneEventChannel, UMovieSceneScriptingEventKey, FMovieSceneEvent>
 {
 	GENERATED_BODY()
 public:
@@ -75,7 +76,7 @@ public:
 	* @return	The key that was created with the specified values at the specified time.
 	*/
 	UFUNCTION(BlueprintCallable, Category = "Editor Scripting | Sequencer Tools | Keys", meta = (DisplayName = "Add Key (Event)"))
-	UMovieSceneScriptingEventKey* AddKey(const FFrameNumber& InTime, FEventPayload NewValue, float SubFrame = 0.f, ESequenceTimeUnit TimeUnit = ESequenceTimeUnit::DisplayRate)
+	UMovieSceneScriptingEventKey* AddKey(const FFrameNumber& InTime, FMovieSceneEvent NewValue, float SubFrame = 0.f, ESequenceTimeUnit TimeUnit = ESequenceTimeUnit::DisplayRate)
 	{
 		return AddKeyInChannel(ChannelHandle, OwningSequence, InTime, NewValue, SubFrame, TimeUnit, EMovieSceneKeyInterpolation::Auto);
 	}
@@ -102,5 +103,5 @@ public:
 
 public:
 	TWeakObjectPtr<UMovieSceneSequence> OwningSequence;
-	TMovieSceneChannelHandle<FMovieSceneEventSectionData> ChannelHandle;
+	TMovieSceneChannelHandle<FMovieSceneEventChannel> ChannelHandle;
 };
\ No newline at end of file
diff --git a/Engine/Source/Editor/CurveAssetEditor/CurveAssetEditor.Build.cs b/Engine/Source/Editor/CurveAssetEditor/CurveAssetEditor.Build.cs
--- a/Engine/Source/Editor/CurveAssetEditor/CurveAssetEditor.Build.cs
+++ b/Engine/Source/Editor/CurveAssetEditor/CurveAssetEditor.Build.cs
@@ -19,7 +19,8 @@ public class CurveAssetEditor : ModuleRules
 				"SlateCore",
                 "EditorStyle",
 				"UnrealEd",
-				"InputCore"
+				"InputCore",
+				"TimeManagement",
 			}
 		);
 
diff --git a/Engine/Source/Editor/CurveAssetEditor/Private/CurveAssetEditor.cpp b/Engine/Source/Editor/CurveAssetEditor/Private/CurveAssetEditor.cpp
--- a/Engine/Source/Editor/CurveAssetEditor/Private/CurveAssetEditor.cpp
+++ b/Engine/Source/Editor/CurveAssetEditor/Private/CurveAssetEditor.cpp
@@ -8,21 +8,86 @@
 #include "EditorStyleSet.h"
 #include "Curves/CurveBase.h"
 #include "CurveAssetEditorModule.h"
-
-#include "SCurveEditor.h"
+#include "CurveEditor.h"
+#include "SCurveEditorPanel.h"
+#include "RichCurveEditorModel.h"
 #include "CurveEditorCommands.h"
-//#include "Toolkits/IToolkitHost.h"
 #include "Widgets/Docking/SDockTab.h"
 #include "Widgets/Input/SNumericDropDown.h"
 #include "Curves/CurveLinearColor.h"
 #include "IDetailsView.h"
 #include "PropertyEditorModule.h"
+#include "Widgets/SFrameRatePicker.h"
+#include "SColorGradientCurveEditorView.h"
+#include "SColorGradientEditor.h"
+#include "CommonFrameRates.h"
+#include "Tree/ICurveEditorTreeItem.h"
+#include "Tree/SCurveEditorTree.h"
+#include "Tree/SCurveEditorTreePin.h"
 
 #define LOCTEXT_NAMESPACE "CurveAssetEditor"
 
 const FName FCurveAssetEditor::CurveTabId( TEXT( "CurveAssetEditor_Curve" ) );
 const FName FCurveAssetEditor::ColorCurveEditorTabId(TEXT("CurveAssetEditor_ColorCurveEditor"));
 
+struct FCurveAssetEditorTreeItem : public ICurveEditorTreeItem
+{
+	FCurveAssetEditorTreeItem(TWeakObjectPtr<UCurveBase> InCurveOwner, const FRichCurveEditInfo& InEditInfo)
+		: CurveOwner(InCurveOwner)
+		, EditInfo(InEditInfo)
+	{
+		if (CurveOwner.IsValid())
+		{
+			CurveName = FText::FromName(EditInfo.CurveName);
+			CurveColor = CurveOwner->GetCurveColor(EditInfo);
+		}
+	}
+
+	virtual TSharedPtr<SWidget> GenerateCurveEditorTreeWidget(const FName& InColumnName, TWeakPtr<FCurveEditor> InCurveEditor, FCurveEditorTreeItemID InTreeItemID) override
+	{
+		if (InColumnName == ColumnNames.Label)
+		{
+			return SNew(SHorizontalBox)
+
+			+ SHorizontalBox::Slot()
+			.Padding(FMargin(0.f, 0.f, 4.f, 0.f))
+			.VAlign(VAlign_Center)
+			.HAlign(HAlign_Right)
+			.AutoWidth()
+			[
+				SNew(STextBlock)
+				.Text(CurveName)
+				.ColorAndOpacity(FSlateColor(CurveColor))
+			];
+		}
+		else if (InColumnName == ColumnNames.PinHeader)
+		{
+			return SNew(SCurveEditorTreePin, InCurveEditor, InTreeItemID);
+		}
+
+		return nullptr;
+	}
+
+	virtual void CreateCurveModels(TArray<TUniquePtr<FCurveModel>>& OutCurveModels) override
+	{
+		if (!CurveOwner.IsValid())
+		{
+			return;
+		}
+
+		TUniquePtr<FRichCurveEditorModel> NewCurve = MakeUnique<FRichCurveEditorModel>(static_cast<FRichCurve*>(EditInfo.CurveToEdit), CurveOwner.Get());
+		NewCurve->SetShortDisplayName(CurveName);
+		NewCurve->SetColor(CurveColor);
+		OutCurveModels.Add(MoveTemp(NewCurve));
+	}
+	
+private:
+	TWeakObjectPtr<UCurveBase> CurveOwner;
+	FRichCurveEditInfo EditInfo;
+	FText CurveName;
+	FLinearColor CurveColor;
+};
+
 void FCurveAssetEditor::RegisterTabSpawners(const TSharedRef<class FTabManager>& InTabManager)
 {
 	WorkspaceMenuCategory = InTabManager->AddLocalWorkspaceMenuCategory(LOCTEXT("WorkspaceMenu_CurveAssetEditor", "Curve Asset Editor"));
@@ -117,7 +182,9 @@ void FCurveAssetEditor::InitCurveAssetEditor( const EToolkitMode::Type Mode, con
 	}
 	const bool bCreateDefaultStandaloneMenu = true;
 	const bool bCreateDefaultToolbar = true;
-	FAssetEditorToolkit::InitAssetEditor( Mode, InitToolkitHost, FCurveAssetEditorModule::CurveAssetEditorAppIdentifier, StandaloneDefaultLayout, bCreateDefaultStandaloneMenu, bCreateDefaultToolbar, CurveToEdit );
+	const bool bToolbarFocusable = false;
+	const bool bUseSmallIcons = true;
+	FAssetEditorToolkit::InitAssetEditor( Mode, InitToolkitHost, FCurveAssetEditorModule::CurveAssetEditorAppIdentifier, StandaloneDefaultLayout, bCreateDefaultStandaloneMenu, bCreateDefaultToolbar, CurveToEdit, bToolbarFocusable, bUseSmallIcons);
 	
 	FCurveAssetEditorModule& CurveAssetEditorModule = FModuleManager::LoadModuleChecked<FCurveAssetEditorModule>( "CurveAssetEditor" );
 	AddMenuExtender(CurveAssetEditorModule.GetMenuExtensibilityManager()->GetAllExtenders(GetToolkitCommands(), GetEditingObjects()));
@@ -131,7 +198,7 @@ void FCurveAssetEditor::InitCurveAssetEditor( const EToolkitMode::Type Mode, con
 		SpawnToolkitTab( CurveTabId, TabInitializationPayload, EToolkitTabSpot::Details );
 	}*/
 
-	if (TrackWidget.IsValid())
+	if (CurveEditor.IsValid())
 	{
 		RegenerateMenusAndToolbars();
 	}
@@ -166,45 +233,79 @@ TSharedRef<SDockTab> FCurveAssetEditor::SpawnTab_CurveAsset( const FSpawnTabArgs
 {
 	check( Args.GetTabId().TabType == CurveTabId );
 
-	ViewMinInput=0.f;
-	ViewMaxInput=5.f;
 
-	InputSnap = 0.1f;
-	OutputSnap = 0.05f;
+	CurveEditor = MakeShared<FCurveEditor>();
+	FCurveEditorInitParams InitParams;
+	CurveEditor->InitCurveEditor(InitParams);
+	CurveEditor->GridLineLabelFormatXAttribute = LOCTEXT("GridXLabelFormat", "{0}");
+
+	// Initialize our bounds at slightly larger than default to avoid clipping the tabs on the color widget.
+	TUniquePtr<ICurveEditorBounds> EditorBounds = MakeUnique<FStaticCurveEditorBounds>();
+	EditorBounds->SetInputBounds(-1.05, 1.05);
+	CurveEditor->SetBounds(MoveTemp(EditorBounds));
+
+	CurveEditorPanel = SNew(SCurveEditorPanel, CurveEditor.ToSharedRef())
+		.TreeContent()
+		[
+			SNew(SCurveEditorTree, CurveEditor)
+		];
 
 	UCurveBase* Curve = Cast<UCurveBase>(GetEditingObject());
+	if (Curve)
+	{
+		check(CurveEditor.IsValid());
+		if (Curve->HasRichCurves())
+		{
+			for (const FRichCurveEditInfo& CurveData : Curve->GetCurves())
+			{		
+				TSharedPtr<FCurveAssetEditorTreeItem> TreeItem = MakeShared<FCurveAssetEditorTreeItem>(Curve, CurveData);
+
+				// Add the channel to the tree-item and let it manage the lifecycle of the tree item.
+				FCurveEditorTreeItem* NewItem = CurveEditor->AddTreeItem(FCurveEditorTreeItemID::Invalid());
+				NewItem->SetStrongItem(TreeItem);
+				
+				// Pin all of the created curves by default for now so that they're visible when you open the
+				// editor. Since there's only ever up to 4 channels we don't have to worry about overwhelming
+				// amounts of curves.
+				for (const FCurveModelID CurveModel : NewItem->GetOrCreateCurves(CurveEditor.Get()))
+				{
+					CurveEditor->PinCurve(CurveModel);
+				}
+			}
+		}
+		else
+		{
+			UE_LOG(LogTemp, Error, TEXT("Non-rich curves are not supported in the Curve Asset editor at this time."));
+		}
+	}
+
+	
 
 	TSharedRef<SDockTab> NewDockTab = SNew(SDockTab)
-		.Icon( FEditorStyle::GetBrush("CurveAssetEditor.Tabs.Properties") )
-		.Label( FText::Format(LOCTEXT("CurveAssetEditorTitle", "{0} Curve Asset"), FText::FromString(GetTabPrefix())))
-		.TabColorScale( GetTabColorScale() )
+		.Icon(FEditorStyle::GetBrush("CurveAssetEditor.Tabs.Properties"))
+		.Label(FText::Format(LOCTEXT("CurveAssetEditorTitle", "{0} Curve Asset"), FText::FromString(GetTabPrefix())))
+		.TabColorScale(GetTabColorScale())
 		[
 			SNew(SBorder)
-			.BorderImage( FEditorStyle::GetBrush("ToolPanel.GroupBorder") )
+			.BorderImage(FEditorStyle::GetBrush("ToolPanel.GroupBorder"))
 			.Padding(0.0f)
 			[
-				SAssignNew(TrackWidget, SCurveEditor)
-				.ViewMinInput(this, &FCurveAssetEditor::GetViewMinInput)
-				.ViewMaxInput(this, &FCurveAssetEditor::GetViewMaxInput)
-				.InputSnap(this, &FCurveAssetEditor::GetInputSnap)
-				.OutputSnap(this, &FCurveAssetEditor::GetOutputSnap)
-				.TimelineLength(this, &FCurveAssetEditor::GetTimelineLength)
-				.OnSetInputViewRange(this, &FCurveAssetEditor::SetInputViewRange)
-				.HideUI(false)
-				.AlwaysDisplayColorCurves(true)
-				.ShowZoomButtons(false)
+				CurveEditorPanel.ToSharedRef()
 			]
 		];
 
-
-
-	FCurveOwnerInterface* CurveOwner = Curve;
-	
-	if (CurveOwner != NULL)
+	// Insert a widget for editing the curve as a Color Gradient if it's a color curve we're editing.
+	UCurveLinearColor* ColorCurve = Cast<UCurveLinearColor>(GetEditingObject());
+	if (ColorCurve)
 	{
-		check(TrackWidget.IsValid());
-		// Set this curve as the SCurveEditor's selected curve
-		TrackWidget->SetCurveOwner(CurveOwner);
+		// We want to register an additional color editing view. This is effectively a "pinned" view who's visibility is controlled by the editor.
+		TSharedRef<SColorGradientCurveEditorView> ColorGradientView = SNew(SColorGradientCurveEditorView, CurveEditor.ToSharedRef())
+			.ViewMinInput(this, &FCurveAssetEditor::GetColorGradientViewMin)
+			.ViewMaxInput(this, &FCurveAssetEditor::GetColorGradientViewMax)
+			.IsEditingEnabled(true);
+
+		ColorGradientView->GetGradientEditor()->SetCurveOwner(ColorCurve);
+		CurveEditorPanel->AddView(ColorGradientView);
 	}
 
 	return NewDockTab;
@@ -226,114 +327,10 @@ TSharedRef<SDockTab> FCurveAssetEditor::SpawnTab_ColorCurveEditor(const FSpawnTa
 	return NewDockTab;
 }
 
-float FCurveAssetEditor::GetInputSnap() const
-{
-	return InputSnap;
-}
-
-void FCurveAssetEditor::SetInputSnap(float value)
-{
-	InputSnap = value;
-}
-
-float FCurveAssetEditor::GetOutputSnap() const
-{
-	return OutputSnap;
-}
-
-void FCurveAssetEditor::SetOutputSnap(float value)
-{
-	OutputSnap = value;
-}
-
-float FCurveAssetEditor::GetTimelineLength() const
-{
-	return 0.f;
-}
-
-void FCurveAssetEditor::SetInputViewRange(float InViewMinInput, float InViewMaxInput)
-{
-	ViewMaxInput = InViewMaxInput;
-	ViewMinInput = InViewMinInput;
-}
-
 TSharedPtr<FExtender> FCurveAssetEditor::GetToolbarExtender()
 {
-	struct Local
-	{
-		static void FillToolbar(FToolBarBuilder& ToolbarBuilder, TSharedRef<SWidget> InputSnapWidget, TSharedRef<SWidget> OutputSnapWidget, FCurveAssetEditor* CurveAssetEditor)
-		{
-			ToolbarBuilder.BeginSection("Curve");
-			{
-				ToolbarBuilder.AddToolBarButton(FCurveEditorCommands::Get().ZoomToFitHorizontal);
-				ToolbarBuilder.AddToolBarButton(FCurveEditorCommands::Get().ZoomToFitVertical);
-				ToolbarBuilder.AddToolBarButton(FCurveEditorCommands::Get().ZoomToFit);
-			}
-			ToolbarBuilder.EndSection();
-
-			ToolbarBuilder.BeginSection("Interpolation");
-			{
-				ToolbarBuilder.AddToolBarButton(FCurveEditorCommands::Get().InterpolationCubicAuto);
-				ToolbarBuilder.AddToolBarButton(FCurveEditorCommands::Get().InterpolationCubicUser);
-				ToolbarBuilder.AddToolBarButton(FCurveEditorCommands::Get().InterpolationCubicBreak);
-				ToolbarBuilder.AddToolBarButton(FCurveEditorCommands::Get().InterpolationLinear);
-				ToolbarBuilder.AddToolBarButton(FCurveEditorCommands::Get().InterpolationConstant);
-			}
-			ToolbarBuilder.EndSection();
-
-			ToolbarBuilder.AddComboButton(
-				FUIAction(),
-				FOnGetContent::CreateSP( CurveAssetEditor, &FCurveAssetEditor::MakeCurveEditorCurveOptionsMenu ),
-				LOCTEXT( "CurveEditorCurveOptions", "Curves Options" ),
-				LOCTEXT( "CurveEditorCurveOptionsToolTip", "Curve Options" ),
-				TAttribute<FSlateIcon>(),
-				true );
-
-			ToolbarBuilder.BeginSection("Snap");
-			{
-				ToolbarBuilder.AddToolBarButton(FCurveEditorCommands::Get().ToggleInputSnapping);
-				ToolbarBuilder.AddWidget(InputSnapWidget);
-				ToolbarBuilder.AddToolBarButton(FCurveEditorCommands::Get().ToggleOutputSnapping);
-				ToolbarBuilder.AddWidget(OutputSnapWidget);
-			}
-			ToolbarBuilder.EndSection();
-		}
-	};
-
-	TSharedPtr<FExtender> ToolbarExtender = MakeShareable(new FExtender);
-
-	TArray<SNumericDropDown<float>::FNamedValue> SnapValues;
-	SnapValues.Add( SNumericDropDown<float>::FNamedValue( 0.001f, LOCTEXT( "Snap_OneThousandth", "0.001" ), LOCTEXT( "SnapDescription_OneThousandth", "Set snap to 1/1000th" ) ) );
-	SnapValues.Add( SNumericDropDown<float>::FNamedValue( 0.01f, LOCTEXT( "Snap_OneHundredth", "0.01" ), LOCTEXT( "SnapDescription_OneHundredth", "Set snap to 1/100th" ) ) );
-	SnapValues.Add( SNumericDropDown<float>::FNamedValue( 0.1f, LOCTEXT( "Snap_OneTenth", "0.1" ), LOCTEXT( "SnapDescription_OneTenth", "Set snap to 1/10th" ) ) );
-	SnapValues.Add( SNumericDropDown<float>::FNamedValue( 1.0f, LOCTEXT( "Snap_One", "1" ), LOCTEXT( "SnapDescription_One", "Set snap to 1" ) ) );
-	SnapValues.Add( SNumericDropDown<float>::FNamedValue( 10.0f, LOCTEXT( "Snap_Ten", "10" ), LOCTEXT( "SnapDescription_Ten", "Set snap to 10" ) ) );
-	SnapValues.Add( SNumericDropDown<float>::FNamedValue( 100.0f, LOCTEXT( "Snap_OneHundred", "100" ), LOCTEXT( "SnapDescription_OneHundred", "Set snap to 100" ) ) );
-
-	TSharedRef<SWidget> InputSnapWidget =
-		SNew( SNumericDropDown<float> )
-		.DropDownValues( SnapValues )
-		.LabelText( LOCTEXT("InputSnapLabel", "Input Snap"))
-		.Value( this, &FCurveAssetEditor::GetInputSnap )
-		.OnValueChanged( this, &FCurveAssetEditor::SetInputSnap )
-		.Orientation( this, &FCurveAssetEditor::GetSnapLabelOrientation );
-
-	TSharedRef<SWidget> OutputSnapWidget =
-		SNew( SNumericDropDown<float> )
-		.DropDownValues( SnapValues )
-		.LabelText( LOCTEXT( "OutputSnapLabel", "Output Snap" ) )
-		.Value( this, &FCurveAssetEditor::GetOutputSnap )
-		.OnValueChanged( this, &FCurveAssetEditor::SetOutputSnap )
-		.Orientation( this, &FCurveAssetEditor::GetSnapLabelOrientation );
-
-	ToolbarExtender->AddToolBarExtension(
-		"Asset",
-		EExtensionHook::After,
-		TrackWidget->GetCommands(),
-		FToolBarExtensionDelegate::CreateStatic(&Local::FillToolbar, InputSnapWidget, OutputSnapWidget, this)
-		);
-
-	return ToolbarExtender;
+	// Use the Curve Editor Panel's extenders which already has all of the icons listed in the right order.
+	return CurveEditorPanel->GetToolbarExtender();
 }
 
 EOrientation FCurveAssetEditor::GetSnapLabelOrientation() const
@@ -374,7 +371,7 @@ TSharedRef<SWidget> FCurveAssetEditor::MakeCurveEditorCurveOptionsMenu()
 		}
 	};
 
-	FMenuBuilder MenuBuilder( true, TrackWidget->GetCommands());
+	FMenuBuilder MenuBuilder( true, CurveEditor->GetCommands());
 
 	MenuBuilder.AddMenuEntry( FCurveEditorCommands::Get().BakeCurve);
 	MenuBuilder.AddMenuEntry( FCurveEditorCommands::Get().ReduceCurve);
@@ -392,4 +389,20 @@ TSharedRef<SWidget> FCurveAssetEditor::MakeCurveEditorCurveOptionsMenu()
 	return MenuBuilder.MakeWidget();
 }
 
+float FCurveAssetEditor::GetColorGradientViewMin() const
+{
+	double MinBounds, MaxBounds;
+	CurveEditor->GetBounds().GetInputBounds(MinBounds, MaxBounds);
+
+	return MinBounds;
+}
+
+float FCurveAssetEditor::GetColorGradientViewMax() const
+{
+	double MinBounds, MaxBounds;
+	CurveEditor->GetBounds().GetInputBounds(MinBounds, MaxBounds);
+
+	return MaxBounds;
+}
+
 #undef LOCTEXT_NAMESPACE
diff --git a/Engine/Source/Editor/CurveAssetEditor/Private/CurveAssetEditor.h b/Engine/Source/Editor/CurveAssetEditor/Private/CurveAssetEditor.h
--- a/Engine/Source/Editor/CurveAssetEditor/Private/CurveAssetEditor.h
+++ b/Engine/Source/Editor/CurveAssetEditor/Private/CurveAssetEditor.h
@@ -7,17 +7,15 @@
 #include "Toolkits/IToolkitHost.h"
 #include "ICurveAssetEditor.h"
 
-class SCurveEditor;
+class FCurveEditor;
 class UCurveBase;
+class SCurveEditorPanel;
 
 class FCurveAssetEditor :  public ICurveAssetEditor
 {
 public:
 
 	FCurveAssetEditor()
-		: TrackWidget( NULL )
-		, ViewMinInput( 0.f )
-		, ViewMaxInput( 0.f )
 	{}
 
 	virtual ~FCurveAssetEditor() {}
@@ -45,43 +43,28 @@ public:
 
 private:
 
-	float GetViewMinInput() const { return ViewMinInput; }
-	float GetViewMaxInput() const { return ViewMaxInput; }
-	/** Sets InViewMinInput and InViewMaxInput */
-	void SetInputViewRange(float InViewMinInput, float InViewMaxInput);
-	/** Return length of timeline */
-	float GetTimelineLength() const;
 	/**	Spawns the tab with the curve asset inside */
 	TSharedRef<SDockTab> SpawnTab_CurveAsset( const FSpawnTabArgs& Args );
 	/**	Spawns the details panel for the color curve */
 	TSharedRef<SDockTab> SpawnTab_ColorCurveEditor(const FSpawnTabArgs& Args);
-	/** Get the snapping value for the input domain. */
-	float GetInputSnap() const;
-	/** Gets the snapping value for the input domain as text. */
-	void SetInputSnap(float value);
-
-	/** Get the snapping value for the output domain. */
-	float GetOutputSnap() const;
-	/** Get set the snapping value for the output domain. */
-	void SetOutputSnap(float value);
 
 	/** Get the orientation for the snap value controls. */
 	EOrientation GetSnapLabelOrientation() const;
 
-	TSharedPtr<class SCurveEditor> TrackWidget;
+	/** Get the view min range from the Bounds to feed to the Color Gradient Editor when we're editing a color curve. */
+	float GetColorGradientViewMin() const;
+	/** Get the view max range from the Bounds to feed to the Color Gradient Editor when we're editing a color curve. */
+	float GetColorGradientViewMax() const;
+
+private:
+	TSharedPtr<FCurveEditor> CurveEditor;
+	TSharedPtr<SCurveEditorPanel> CurveEditorPanel;
 
 	/**	The tab id for the curve asset tab */
 	static const FName CurveTabId;
 	/**	The tab id for the color curve editor tab */
 	static const FName ColorCurveEditorTabId;
 
-	float ViewMinInput;
-	float ViewMaxInput;
-
-	/** The snapping value for the input domain. */
-	float InputSnap;
-	/** The snapping value for the output domain. */
-	float OutputSnap;
 	/* Holds the details panel for the color curve */
 	TSharedPtr<class IDetailsView> ColorCurveDetailsView;
 };
diff --git a/Engine/Source/Editor/CurveEditor/CurveEditor.Build.cs b/Engine/Source/Editor/CurveEditor/CurveEditor.Build.cs
--- a/Engine/Source/Editor/CurveEditor/CurveEditor.Build.cs
+++ b/Engine/Source/Editor/CurveEditor/CurveEditor.Build.cs
@@ -17,7 +17,10 @@ public class CurveEditor : ModuleRules
 				"SlateCore",
 				"TimeManagement",
 				"UnrealEd",
+				"SequencerWidgets",
 			}
 		);
+
+        PublicDependencyModuleNames.Add("SequencerWidgets");
 	}
 }
