Curve Editor: Update copyright notices

diff --git a/Engine/Plugins/Editor/CurveEditorTools/Source/CurveEditorTools/CurveEditorTools.Build.cs b/Engine/Plugins/Editor/CurveEditorTools/Source/CurveEditorTools/CurveEditorTools.Build.cs
--- a/Engine/Plugins/Editor/CurveEditorTools/Source/CurveEditorTools/CurveEditorTools.Build.cs
+++ b/Engine/Plugins/Editor/CurveEditorTools/Source/CurveEditorTools/CurveEditorTools.Build.cs
@@ -1,4 +1,4 @@
-// Copyright 1998-2018 Epic Games, Inc. All Rights Reserved.
+// Copyright 1998-2019 Epic Games, Inc. All Rights Reserved.
 
 using UnrealBuildTool;
 
diff --git a/Engine/Plugins/Editor/CurveEditorTools/Source/CurveEditorTools/Private/CurveEditorFocusExtension.cpp b/Engine/Plugins/Editor/CurveEditorTools/Source/CurveEditorTools/Private/CurveEditorFocusExtension.cpp
--- a/Engine/Plugins/Editor/CurveEditorTools/Source/CurveEditorTools/Private/CurveEditorFocusExtension.cpp
+++ b/Engine/Plugins/Editor/CurveEditorTools/Source/CurveEditorTools/Private/CurveEditorFocusExtension.cpp
@@ -1,4 +1,4 @@
-// Copyright 1998-2018 Epic Games, Inc. All Rights Reserved.
+// Copyright 1998-2019 Epic Games, Inc. All Rights Reserved.
 
 #include "CurveEditorFocusExtension.h"
 #include "Framework/MultiBox/MultiBoxBuilder.h"
diff --git a/Engine/Plugins/Editor/CurveEditorTools/Source/CurveEditorTools/Private/CurveEditorFocusExtension.h b/Engine/Plugins/Editor/CurveEditorTools/Source/CurveEditorTools/Private/CurveEditorFocusExtension.h
--- a/Engine/Plugins/Editor/CurveEditorTools/Source/CurveEditorTools/Private/CurveEditorFocusExtension.h
+++ b/Engine/Plugins/Editor/CurveEditorTools/Source/CurveEditorTools/Private/CurveEditorFocusExtension.h
@@ -1,4 +1,4 @@
-// Copyright 1998-2018 Epic Games, Inc. All Rights Reserved.
+// Copyright 1998-2019 Epic Games, Inc. All Rights Reserved.
 
 #pragma once
 
diff --git a/Engine/Plugins/Editor/CurveEditorTools/Source/CurveEditorTools/Private/CurveEditorToolCommands.cpp b/Engine/Plugins/Editor/CurveEditorTools/Source/CurveEditorTools/Private/CurveEditorToolCommands.cpp
--- a/Engine/Plugins/Editor/CurveEditorTools/Source/CurveEditorTools/Private/CurveEditorToolCommands.cpp
+++ b/Engine/Plugins/Editor/CurveEditorTools/Source/CurveEditorTools/Private/CurveEditorToolCommands.cpp
@@ -1,4 +1,4 @@
-// Copyright 1998-2018 Epic Games, Inc. All Rights Reserved.
+// Copyright 1998-2019 Epic Games, Inc. All Rights Reserved.
 
 #include "CurveEditorToolCommands.h"
 
diff --git a/Engine/Plugins/Editor/CurveEditorTools/Source/CurveEditorTools/Private/CurveEditorToolCommands.h b/Engine/Plugins/Editor/CurveEditorTools/Source/CurveEditorTools/Private/CurveEditorToolCommands.h
--- a/Engine/Plugins/Editor/CurveEditorTools/Source/CurveEditorTools/Private/CurveEditorToolCommands.h
+++ b/Engine/Plugins/Editor/CurveEditorTools/Source/CurveEditorTools/Private/CurveEditorToolCommands.h
@@ -1,4 +1,4 @@
-// Copyright 1998-2018 Epic Games, Inc. All Rights Reserved.
+// Copyright 1998-2019 Epic Games, Inc. All Rights Reserved.
 
 #pragma once
 
diff --git a/Engine/Plugins/Editor/CurveEditorTools/Source/CurveEditorTools/Private/CurveEditorToolsModule.cpp b/Engine/Plugins/Editor/CurveEditorTools/Source/CurveEditorTools/Private/CurveEditorToolsModule.cpp
--- a/Engine/Plugins/Editor/CurveEditorTools/Source/CurveEditorTools/Private/CurveEditorToolsModule.cpp
+++ b/Engine/Plugins/Editor/CurveEditorTools/Source/CurveEditorTools/Private/CurveEditorToolsModule.cpp
@@ -1,4 +1,4 @@
-// Copyright 1998-2018 Epic Games, Inc. All Rights Reserved.
+// Copyright 1998-2019 Epic Games, Inc. All Rights Reserved.
 
 #include "Modules/ModuleManager.h"
 #include "ICurveEditorModule.h"
diff --git a/Engine/Plugins/Editor/CurveEditorTools/Source/CurveEditorTools/Private/Filters/CurveEditorFFTFilter.cpp b/Engine/Plugins/Editor/CurveEditorTools/Source/CurveEditorTools/Private/Filters/CurveEditorFFTFilter.cpp
--- a/Engine/Plugins/Editor/CurveEditorTools/Source/CurveEditorTools/Private/Filters/CurveEditorFFTFilter.cpp
+++ b/Engine/Plugins/Editor/CurveEditorTools/Source/CurveEditorTools/Private/Filters/CurveEditorFFTFilter.cpp
@@ -1,4 +1,4 @@
-// Copyright 1998-2018 Epic Games, Inc. All Rights Reserved.
+// Copyright 1998-2019 Epic Games, Inc. All Rights Reserved.
 
 #include "CurveEditorFFTFilter.h"
 #include "DSP/PassiveFilter.h"
diff --git a/Engine/Plugins/Editor/CurveEditorTools/Source/CurveEditorTools/Private/Filters/CurveEditorFFTFilter.h b/Engine/Plugins/Editor/CurveEditorTools/Source/CurveEditorTools/Private/Filters/CurveEditorFFTFilter.h
--- a/Engine/Plugins/Editor/CurveEditorTools/Source/CurveEditorTools/Private/Filters/CurveEditorFFTFilter.h
+++ b/Engine/Plugins/Editor/CurveEditorTools/Source/CurveEditorTools/Private/Filters/CurveEditorFFTFilter.h
@@ -1,4 +1,4 @@
-// Copyright 1998-2018 Epic Games, Inc. All Rights Reserved.
+// Copyright 1998-2019 Epic Games, Inc. All Rights Reserved.
 
 #pragma once
 
diff --git a/Engine/Plugins/Editor/CurveEditorTools/Source/CurveEditorTools/Private/Tools/CurveEditorRetimeTool.cpp b/Engine/Plugins/Editor/CurveEditorTools/Source/CurveEditorTools/Private/Tools/CurveEditorRetimeTool.cpp
--- a/Engine/Plugins/Editor/CurveEditorTools/Source/CurveEditorTools/Private/Tools/CurveEditorRetimeTool.cpp
+++ b/Engine/Plugins/Editor/CurveEditorTools/Source/CurveEditorTools/Private/Tools/CurveEditorRetimeTool.cpp
@@ -1,4 +1,4 @@
-// Copyright 1998-2018 Epic Games, Inc. All Rights Reserved.
+// Copyright 1998-2019 Epic Games, Inc. All Rights Reserved.
 
 #include "CurveEditorRetimeTool.h"
 #include "CurveEditorToolCommands.h"
diff --git a/Engine/Plugins/Editor/CurveEditorTools/Source/CurveEditorTools/Private/Tools/CurveEditorRetimeTool.h b/Engine/Plugins/Editor/CurveEditorTools/Source/CurveEditorTools/Private/Tools/CurveEditorRetimeTool.h
--- a/Engine/Plugins/Editor/CurveEditorTools/Source/CurveEditorTools/Private/Tools/CurveEditorRetimeTool.h
+++ b/Engine/Plugins/Editor/CurveEditorTools/Source/CurveEditorTools/Private/Tools/CurveEditorRetimeTool.h
@@ -1,4 +1,4 @@
-// Copyright 1998-2018 Epic Games, Inc. All Rights Reserved.
+// Copyright 1998-2019 Epic Games, Inc. All Rights Reserved.
 
 #pragma once
 #include "CoreMinimal.h"
diff --git a/Engine/Plugins/Editor/CurveEditorTools/Source/CurveEditorTools/Private/Tools/CurveEditorTransformTool.cpp b/Engine/Plugins/Editor/CurveEditorTools/Source/CurveEditorTools/Private/Tools/CurveEditorTransformTool.cpp
--- a/Engine/Plugins/Editor/CurveEditorTools/Source/CurveEditorTools/Private/Tools/CurveEditorTransformTool.cpp
+++ b/Engine/Plugins/Editor/CurveEditorTools/Source/CurveEditorTools/Private/Tools/CurveEditorTransformTool.cpp
@@ -1,4 +1,4 @@
-// Copyright 1998-2018 Epic Games, Inc. All Rights Reserved.
+// Copyright 1998-2019 Epic Games, Inc. All Rights Reserved.
 
 #include "CurveEditorTransformTool.h"
 #include "CurveEditorToolCommands.h"
diff --git a/Engine/Plugins/Editor/CurveEditorTools/Source/CurveEditorTools/Private/Tools/CurveEditorTransformTool.h b/Engine/Plugins/Editor/CurveEditorTools/Source/CurveEditorTools/Private/Tools/CurveEditorTransformTool.h
--- a/Engine/Plugins/Editor/CurveEditorTools/Source/CurveEditorTools/Private/Tools/CurveEditorTransformTool.h
+++ b/Engine/Plugins/Editor/CurveEditorTools/Source/CurveEditorTools/Private/Tools/CurveEditorTransformTool.h
@@ -1,4 +1,4 @@
-// Copyright 1998-2018 Epic Games, Inc. All Rights Reserved.
+// Copyright 1998-2019 Epic Games, Inc. All Rights Reserved.
 
 #pragma once
 #include "Templates/SharedPointer.h"
diff --git a/Engine/Source/Editor/CurveAssetEditor/Private/SColorGradientCurveEditorView.cpp b/Engine/Source/Editor/CurveAssetEditor/Private/SColorGradientCurveEditorView.cpp
--- a/Engine/Source/Editor/CurveAssetEditor/Private/SColorGradientCurveEditorView.cpp
+++ b/Engine/Source/Editor/CurveAssetEditor/Private/SColorGradientCurveEditorView.cpp
@@ -1,3 +1,5 @@
+// Copyright 1998-2019 Epic Games, Inc. All Rights Reserved.
+
 #include "SColorGradientCurveEditorView.h"
 #include "EditorStyleSet.h"
 #include "Widgets/Layout/SBorder.h"
diff --git a/Engine/Source/Editor/CurveAssetEditor/Private/SColorGradientCurveEditorView.h b/Engine/Source/Editor/CurveAssetEditor/Private/SColorGradientCurveEditorView.h
--- a/Engine/Source/Editor/CurveAssetEditor/Private/SColorGradientCurveEditorView.h
+++ b/Engine/Source/Editor/CurveAssetEditor/Private/SColorGradientCurveEditorView.h
@@ -1,4 +1,4 @@
-// Copyright 1998-2018 Epic Games, Inc. All Rights Reserved.
+// Copyright 1998-2019 Epic Games, Inc. All Rights Reserved.
 
 #pragma once
 
diff --git a/Engine/Source/Editor/CurveEditor/Private/Filters/CurveEditorBakeFilter.cpp b/Engine/Source/Editor/CurveEditor/Private/Filters/CurveEditorBakeFilter.cpp
--- a/Engine/Source/Editor/CurveEditor/Private/Filters/CurveEditorBakeFilter.cpp
+++ b/Engine/Source/Editor/CurveEditor/Private/Filters/CurveEditorBakeFilter.cpp
@@ -1,4 +1,4 @@
-// Copyright 1998-2018 Epic Games, Inc. All Rights Reserved.
+// Copyright 1998-2019 Epic Games, Inc. All Rights Reserved.
 
 #include "Filters/CurveEditorBakeFilter.h"
 #include "CoreMinimal.h"
diff --git a/Engine/Source/Editor/CurveEditor/Private/Filters/CurveEditorReduceFilter.cpp b/Engine/Source/Editor/CurveEditor/Private/Filters/CurveEditorReduceFilter.cpp
--- a/Engine/Source/Editor/CurveEditor/Private/Filters/CurveEditorReduceFilter.cpp
+++ b/Engine/Source/Editor/CurveEditor/Private/Filters/CurveEditorReduceFilter.cpp
@@ -1,4 +1,4 @@
-// Copyright 1998-2018 Epic Games, Inc. All Rights Reserved.
+// Copyright 1998-2019 Epic Games, Inc. All Rights Reserved.
 
 #include "Filters/CurveEditorReduceFilter.h"
 #include "CurveEditor.h"
diff --git a/Engine/Source/Editor/CurveEditor/Private/RichCurveEditorModel.cpp b/Engine/Source/Editor/CurveEditor/Private/RichCurveEditorModel.cpp
--- a/Engine/Source/Editor/CurveEditor/Private/RichCurveEditorModel.cpp
+++ b/Engine/Source/Editor/CurveEditor/Private/RichCurveEditorModel.cpp
@@ -1,4 +1,4 @@
-// Copyright 1998-2018 Epic Games, Inc. All Rights Reserved.
+// Copyright 1998-2019 Epic Games, Inc. All Rights Reserved.
 
 #include "RichCurveEditorModel.h"
 #include "RichCurveKeyProxy.h"
diff --git a/Engine/Source/Editor/CurveEditor/Private/RichCurveKeyProxy.h b/Engine/Source/Editor/CurveEditor/Private/RichCurveKeyProxy.h
--- a/Engine/Source/Editor/CurveEditor/Private/RichCurveKeyProxy.h
+++ b/Engine/Source/Editor/CurveEditor/Private/RichCurveKeyProxy.h
@@ -1,4 +1,4 @@
-// Copyright 1998-2018 Epic Games, Inc. All Rights Reserved.
+// Copyright 1998-2019 Epic Games, Inc. All Rights Reserved.
 
 #pragma once
 
diff --git a/Engine/Source/Editor/CurveEditor/Private/SCurveEditorView.cpp b/Engine/Source/Editor/CurveEditor/Private/SCurveEditorView.cpp
--- a/Engine/Source/Editor/CurveEditor/Private/SCurveEditorView.cpp
+++ b/Engine/Source/Editor/CurveEditor/Private/SCurveEditorView.cpp
@@ -1,4 +1,4 @@
-// Copyright 1998-2018 Epic Games, Inc. All Rights Reserved.
+// Copyright 1998-2019 Epic Games, Inc. All Rights Reserved.
 
 #include "SCurveEditorView.h"
 #include "ICurveEditorBounds.h"
diff --git a/Engine/Source/Editor/CurveEditor/Private/SCurveKeyDetailPanel.cpp b/Engine/Source/Editor/CurveEditor/Private/SCurveKeyDetailPanel.cpp
--- a/Engine/Source/Editor/CurveEditor/Private/SCurveKeyDetailPanel.cpp
+++ b/Engine/Source/Editor/CurveEditor/Private/SCurveKeyDetailPanel.cpp
@@ -1,4 +1,4 @@
-// Copyright 1998-2018 Epic Games, Inc. All Rights Reserved.
+// Copyright 1998-2019 Epic Games, Inc. All Rights Reserved.
 
 #include "SCurveKeyDetailPanel.h"
 #include "Modules/ModuleManager.h"
diff --git a/Engine/Source/Editor/CurveEditor/Private/Tree/ICurveEditorTreeItem.cpp b/Engine/Source/Editor/CurveEditor/Private/Tree/ICurveEditorTreeItem.cpp
--- a/Engine/Source/Editor/CurveEditor/Private/Tree/ICurveEditorTreeItem.cpp
+++ b/Engine/Source/Editor/CurveEditor/Private/Tree/ICurveEditorTreeItem.cpp
@@ -1,3 +1,4 @@
+// Copyright 1998-2019 Epic Games, Inc. All Rights Reserved.
 
 #include "Tree/ICurveEditorTreeItem.h"
 
diff --git a/Engine/Source/Editor/CurveEditor/Private/Views/SCurveEditorViewAbsolute.cpp b/Engine/Source/Editor/CurveEditor/Private/Views/SCurveEditorViewAbsolute.cpp
--- a/Engine/Source/Editor/CurveEditor/Private/Views/SCurveEditorViewAbsolute.cpp
+++ b/Engine/Source/Editor/CurveEditor/Private/Views/SCurveEditorViewAbsolute.cpp
@@ -1,4 +1,4 @@
-// Copyright 1998-2018 Epic Games, Inc. All Rights Reserved.
+// Copyright 1998-2019 Epic Games, Inc. All Rights Reserved.
 
 #include "Views/SCurveEditorViewAbsolute.h"
 #include "Widgets/Text/STextBlock.h"
diff --git a/Engine/Source/Editor/CurveEditor/Private/Views/SCurveEditorViewAbsolute.h b/Engine/Source/Editor/CurveEditor/Private/Views/SCurveEditorViewAbsolute.h
--- a/Engine/Source/Editor/CurveEditor/Private/Views/SCurveEditorViewAbsolute.h
+++ b/Engine/Source/Editor/CurveEditor/Private/Views/SCurveEditorViewAbsolute.h
@@ -1,4 +1,4 @@
-// Copyright 1998-2018 Epic Games, Inc. All Rights Reserved.
+// Copyright 1998-2019 Epic Games, Inc. All Rights Reserved.
 
 #pragma once
 
diff --git a/Engine/Source/Editor/CurveEditor/Private/Views/SCurveEditorViewNormalized.cpp b/Engine/Source/Editor/CurveEditor/Private/Views/SCurveEditorViewNormalized.cpp
--- a/Engine/Source/Editor/CurveEditor/Private/Views/SCurveEditorViewNormalized.cpp
+++ b/Engine/Source/Editor/CurveEditor/Private/Views/SCurveEditorViewNormalized.cpp
@@ -1,4 +1,4 @@
-// Copyright 1998-2018 Epic Games, Inc. All Rights Reserved.
+// Copyright 1998-2019 Epic Games, Inc. All Rights Reserved.
 
 #include "Views/SCurveEditorViewNormalized.h"
 #include "CurveEditor.h"
diff --git a/Engine/Source/Editor/CurveEditor/Private/Views/SCurveEditorViewNormalized.h b/Engine/Source/Editor/CurveEditor/Private/Views/SCurveEditorViewNormalized.h
--- a/Engine/Source/Editor/CurveEditor/Private/Views/SCurveEditorViewNormalized.h
+++ b/Engine/Source/Editor/CurveEditor/Private/Views/SCurveEditorViewNormalized.h
@@ -1,4 +1,4 @@
-// Copyright 1998-2018 Epic Games, Inc. All Rights Reserved.
+// Copyright 1998-2019 Epic Games, Inc. All Rights Reserved.
 
 #pragma once
 
diff --git a/Engine/Source/Editor/CurveEditor/Private/Views/SCurveEditorViewStacked.cpp b/Engine/Source/Editor/CurveEditor/Private/Views/SCurveEditorViewStacked.cpp
--- a/Engine/Source/Editor/CurveEditor/Private/Views/SCurveEditorViewStacked.cpp
+++ b/Engine/Source/Editor/CurveEditor/Private/Views/SCurveEditorViewStacked.cpp
@@ -1,4 +1,4 @@
-// Copyright 1998-2018 Epic Games, Inc. All Rights Reserved.
+// Copyright 1998-2019 Epic Games, Inc. All Rights Reserved.
 
 #include "Views/SCurveEditorViewStacked.h"
 #include "CurveEditor.h"
diff --git a/Engine/Source/Editor/CurveEditor/Private/Views/SCurveEditorViewStacked.h b/Engine/Source/Editor/CurveEditor/Private/Views/SCurveEditorViewStacked.h
--- a/Engine/Source/Editor/CurveEditor/Private/Views/SCurveEditorViewStacked.h
+++ b/Engine/Source/Editor/CurveEditor/Private/Views/SCurveEditorViewStacked.h
@@ -1,4 +1,4 @@
-// Copyright 1998-2018 Epic Games, Inc. All Rights Reserved.
+// Copyright 1998-2019 Epic Games, Inc. All Rights Reserved.
 
 #pragma once
 
diff --git a/Engine/Source/Editor/CurveEditor/Private/Views/SInteractiveCurveEditorView.cpp b/Engine/Source/Editor/CurveEditor/Private/Views/SInteractiveCurveEditorView.cpp
--- a/Engine/Source/Editor/CurveEditor/Private/Views/SInteractiveCurveEditorView.cpp
+++ b/Engine/Source/Editor/CurveEditor/Private/Views/SInteractiveCurveEditorView.cpp
@@ -1,4 +1,4 @@
-// Copyright 1998-2018 Epic Games, Inc. All Rights Reserved.
+// Copyright 1998-2019 Epic Games, Inc. All Rights Reserved.
 
 #include "Views/SInteractiveCurveEditorView.h"
 #include "EditorStyleSet.h"
diff --git a/Engine/Source/Editor/CurveEditor/Public/Filters/CurveEditorBakeFilter.h b/Engine/Source/Editor/CurveEditor/Public/Filters/CurveEditorBakeFilter.h
--- a/Engine/Source/Editor/CurveEditor/Public/Filters/CurveEditorBakeFilter.h
+++ b/Engine/Source/Editor/CurveEditor/Public/Filters/CurveEditorBakeFilter.h
@@ -1,4 +1,4 @@
-// Copyright 1998-2018 Epic Games, Inc. All Rights Reserved.
+// Copyright 1998-2019 Epic Games, Inc. All Rights Reserved.
 
 #pragma once
 
diff --git a/Engine/Source/Editor/CurveEditor/Public/Filters/CurveEditorFilterBase.h b/Engine/Source/Editor/CurveEditor/Public/Filters/CurveEditorFilterBase.h
--- a/Engine/Source/Editor/CurveEditor/Public/Filters/CurveEditorFilterBase.h
+++ b/Engine/Source/Editor/CurveEditor/Public/Filters/CurveEditorFilterBase.h
@@ -1,4 +1,4 @@
-// Copyright 1998-2018 Epic Games, Inc. All Rights Reserved.
+// Copyright 1998-2019 Epic Games, Inc. All Rights Reserved.
 
 #pragma once
 
diff --git a/Engine/Source/Editor/CurveEditor/Public/Filters/CurveEditorReduceFilter.h b/Engine/Source/Editor/CurveEditor/Public/Filters/CurveEditorReduceFilter.h
--- a/Engine/Source/Editor/CurveEditor/Public/Filters/CurveEditorReduceFilter.h
+++ b/Engine/Source/Editor/CurveEditor/Public/Filters/CurveEditorReduceFilter.h
@@ -1,4 +1,4 @@
-// Copyright 1998-2018 Epic Games, Inc. All Rights Reserved.
+// Copyright 1998-2019 Epic Games, Inc. All Rights Reserved.
 
 #pragma once
 
diff --git a/Engine/Source/Editor/CurveEditor/Public/Filters/SCurveEditorFilterPanel.cpp b/Engine/Source/Editor/CurveEditor/Public/Filters/SCurveEditorFilterPanel.cpp
--- a/Engine/Source/Editor/CurveEditor/Public/Filters/SCurveEditorFilterPanel.cpp
+++ b/Engine/Source/Editor/CurveEditor/Public/Filters/SCurveEditorFilterPanel.cpp
@@ -1,4 +1,4 @@
-// Copyright 1998-2018 Epic Games, Inc. All Rights Reserved.
+// Copyright 1998-2019 Epic Games, Inc. All Rights Reserved.
 
 #include "SCurveEditorFilterPanel.h"
 #include "Widgets/Input/SButton.h"
diff --git a/Engine/Source/Editor/CurveEditor/Public/Filters/SCurveEditorFilterPanel.h b/Engine/Source/Editor/CurveEditor/Public/Filters/SCurveEditorFilterPanel.h
--- a/Engine/Source/Editor/CurveEditor/Public/Filters/SCurveEditorFilterPanel.h
+++ b/Engine/Source/Editor/CurveEditor/Public/Filters/SCurveEditorFilterPanel.h
@@ -1,4 +1,4 @@
-// Copyright 1998-2018 Epic Games, Inc. All Rights Reserved.
+// Copyright 1998-2019 Epic Games, Inc. All Rights Reserved.
 
 #pragma once
 
diff --git a/Engine/Source/Editor/CurveEditor/Public/ICurveEditorExtension.h b/Engine/Source/Editor/CurveEditor/Public/ICurveEditorExtension.h
--- a/Engine/Source/Editor/CurveEditor/Public/ICurveEditorExtension.h
+++ b/Engine/Source/Editor/CurveEditor/Public/ICurveEditorExtension.h
@@ -1,4 +1,4 @@
-// Copyright 1998-2018 Epic Games, Inc. All Rights Reserved.
+// Copyright 1998-2019 Epic Games, Inc. All Rights Reserved.
 
 #pragma once
 #include "Templates/SharedPointer.h"
diff --git a/Engine/Source/Editor/CurveEditor/Public/ICurveEditorModule.h b/Engine/Source/Editor/CurveEditor/Public/ICurveEditorModule.h
--- a/Engine/Source/Editor/CurveEditor/Public/ICurveEditorModule.h
+++ b/Engine/Source/Editor/CurveEditor/Public/ICurveEditorModule.h
@@ -1,4 +1,4 @@
-// Copyright 1998-2018 Epic Games, Inc. All Rights Reserved.
+// Copyright 1998-2019 Epic Games, Inc. All Rights Reserved.
 
 #pragma  once
 
diff --git a/Engine/Source/Editor/CurveEditor/Public/ICurveEditorToolExtension.h b/Engine/Source/Editor/CurveEditor/Public/ICurveEditorToolExtension.h
--- a/Engine/Source/Editor/CurveEditor/Public/ICurveEditorToolExtension.h
+++ b/Engine/Source/Editor/CurveEditor/Public/ICurveEditorToolExtension.h
@@ -1,4 +1,4 @@
-// Copyright 1998-2018 Epic Games, Inc. All Rights Reserved.
+// Copyright 1998-2019 Epic Games, Inc. All Rights Reserved.
 
 #pragma once
 #include "Templates/SharedPointer.h"
diff --git a/Engine/Source/Editor/CurveEditor/Public/RichCurveEditorModel.h b/Engine/Source/Editor/CurveEditor/Public/RichCurveEditorModel.h
--- a/Engine/Source/Editor/CurveEditor/Public/RichCurveEditorModel.h
+++ b/Engine/Source/Editor/CurveEditor/Public/RichCurveEditorModel.h
@@ -1,4 +1,4 @@
-// Copyright 1998-2018 Epic Games, Inc. All Rights Reserved.
+// Copyright 1998-2019 Epic Games, Inc. All Rights Reserved.
 
 #pragma once
 
diff --git a/Engine/Source/Editor/CurveEditor/Public/SCurveEditorView.h b/Engine/Source/Editor/CurveEditor/Public/SCurveEditorView.h
--- a/Engine/Source/Editor/CurveEditor/Public/SCurveEditorView.h
+++ b/Engine/Source/Editor/CurveEditor/Public/SCurveEditorView.h
@@ -1,4 +1,4 @@
-// Copyright 1998-2018 Epic Games, Inc. All Rights Reserved.
+// Copyright 1998-2019 Epic Games, Inc. All Rights Reserved.
 
 #pragma once
 
diff --git a/Engine/Source/Editor/CurveEditor/Public/SCurveKeyDetailPanel.h b/Engine/Source/Editor/CurveEditor/Public/SCurveKeyDetailPanel.h
--- a/Engine/Source/Editor/CurveEditor/Public/SCurveKeyDetailPanel.h
+++ b/Engine/Source/Editor/CurveEditor/Public/SCurveKeyDetailPanel.h
@@ -1,4 +1,4 @@
-// Copyright 1998-2018 Epic Games, Inc. All Rights Reserved.
+// Copyright 1998-2019 Epic Games, Inc. All Rights Reserved.
 
 #pragma once
 
diff --git a/Engine/Source/Editor/CurveEditor/Public/Tree/ICurveEditorTreeItem.h b/Engine/Source/Editor/CurveEditor/Public/Tree/ICurveEditorTreeItem.h
--- a/Engine/Source/Editor/CurveEditor/Public/Tree/ICurveEditorTreeItem.h
+++ b/Engine/Source/Editor/CurveEditor/Public/Tree/ICurveEditorTreeItem.h
@@ -1,3 +1,4 @@
+// Copyright 1998-2019 Epic Games, Inc. All Rights Reserved.
 
 #pragma once
 
diff --git a/Engine/Source/Editor/CurveEditor/Public/Views/SInteractiveCurveEditorView.h b/Engine/Source/Editor/CurveEditor/Public/Views/SInteractiveCurveEditorView.h
--- a/Engine/Source/Editor/CurveEditor/Public/Views/SInteractiveCurveEditorView.h
+++ b/Engine/Source/Editor/CurveEditor/Public/Views/SInteractiveCurveEditorView.h
@@ -1,4 +1,4 @@
-// Copyright 1998-2018 Epic Games, Inc. All Rights Reserved.
+// Copyright 1998-2019 Epic Games, Inc. All Rights Reserved.
 
 #pragma once
 
diff --git a/Engine/Source/Editor/MovieSceneTools/Private/Channels/EventChannelCurveModel.cpp b/Engine/Source/Editor/MovieSceneTools/Private/Channels/EventChannelCurveModel.cpp
--- a/Engine/Source/Editor/MovieSceneTools/Private/Channels/EventChannelCurveModel.cpp
+++ b/Engine/Source/Editor/MovieSceneTools/Private/Channels/EventChannelCurveModel.cpp
@@ -1,4 +1,4 @@
-// Copyright 1998-2018 Epic Games, Inc. All Rights Reserved.
+// Copyright 1998-2019 Epic Games, Inc. All Rights Reserved.
 
 #include "EventChannelCurveModel.h"
 #include "Math/Vector2D.h"
diff --git a/Engine/Source/Editor/MovieSceneTools/Private/Channels/EventChannelCurveModel.h b/Engine/Source/Editor/MovieSceneTools/Private/Channels/EventChannelCurveModel.h
--- a/Engine/Source/Editor/MovieSceneTools/Private/Channels/EventChannelCurveModel.h
+++ b/Engine/Source/Editor/MovieSceneTools/Private/Channels/EventChannelCurveModel.h
@@ -1,4 +1,4 @@
-// Copyright 1998-2018 Epic Games, Inc. All Rights Reserved.
+// Copyright 1998-2019 Epic Games, Inc. All Rights Reserved.
 
 #pragma once
 
diff --git a/Engine/Source/Editor/MovieSceneTools/Private/Channels/SCurveEditorEventChannelView.cpp b/Engine/Source/Editor/MovieSceneTools/Private/Channels/SCurveEditorEventChannelView.cpp
--- a/Engine/Source/Editor/MovieSceneTools/Private/Channels/SCurveEditorEventChannelView.cpp
+++ b/Engine/Source/Editor/MovieSceneTools/Private/Channels/SCurveEditorEventChannelView.cpp
@@ -1,4 +1,4 @@
-// Copyright 1998-2018 Epic Games, Inc. All Rights Reserved.
+// Copyright 1998-2019 Epic Games, Inc. All Rights Reserved.
 
 #include "SCurveEditorEventChannelView.h"
 #include "Rendering/DrawElements.h"
diff --git a/Engine/Source/Editor/MovieSceneTools/Private/Channels/SCurveEditorEventChannelView.h b/Engine/Source/Editor/MovieSceneTools/Private/Channels/SCurveEditorEventChannelView.h
--- a/Engine/Source/Editor/MovieSceneTools/Private/Channels/SCurveEditorEventChannelView.h
+++ b/Engine/Source/Editor/MovieSceneTools/Private/Channels/SCurveEditorEventChannelView.h
@@ -1,4 +1,4 @@
-// Copyright 1998-2018 Epic Games, Inc. All Rights Reserved.
+// Copyright 1998-2019 Epic Games, Inc. All Rights Reserved.
 
 #pragma once
 
