Curve Editor: Fix bug in transform tool where stacked view bounds were continuously updated during drag operation

diff --git a/Engine/Plugins/Editor/CurveEditorTools/Source/CurveEditorTools/Private/Tools/CurveEditorTransformTool.cpp b/Engine/Plugins/Editor/CurveEditorTools/Source/CurveEditorTools/Private/Tools/CurveEditorTransformTool.cpp
--- a/Engine/Plugins/Editor/CurveEditorTools/Source/CurveEditorTools/Private/Tools/CurveEditorTransformTool.cpp
+++ b/Engine/Plugins/Editor/CurveEditorTools/Source/CurveEditorTools/Private/Tools/CurveEditorTransformTool.cpp
@@ -412,14 +412,16 @@ void FCurveEditorTransformTool::DrawMarqueeWidget(const FCurveEditorTransformWid
 
 void FCurveEditorTransformTool::OnDragStart()
 {
-	ActiveTransaction = MakeUnique<FScopedTransaction>(TEXT("CurveEditorTransformTool"), LOCTEXT("CurveEditorTransformToolTransaction", "Transform Key(s)"), nullptr);
-
 	TSharedPtr<FCurveEditor> CurveEditor = WeakCurveEditor.Pin();
 	if (!CurveEditor)
 	{
 		return;
 	}
 
+	ActiveTransaction = MakeUnique<FScopedTransaction>(TEXT("CurveEditorTransformTool"), LOCTEXT("CurveEditorTransformToolTransaction", "Transform Key(s)"), nullptr);
+
+	CurveEditor->SuppressBoundTransformUpdates(true);
+
 	// We need to cache our key data because all of our calculations have to be relative to the starting data and not the current per-frame data.
 	KeysByCurve.Reset();
 	for (const TTuple<FCurveModelID, FKeyHandleSet>& Pair : CurveEditor->GetSelection().GetAll())
@@ -619,6 +621,14 @@ void FCurveEditorTransformTool::OnDrag(const FPointerEvent& MouseEvent)
 
 void FCurveEditorTransformTool::OnDragEnd()
 {
+	TSharedPtr<FCurveEditor> CurveEditor = WeakCurveEditor.Pin();
+	if (!CurveEditor.IsValid())
+	{
+		return;
+	}
+
+	CurveEditor->SuppressBoundTransformUpdates(false);
+
 	// This finalizes the transaction
 	ActiveTransaction.Reset();
 	// PreDragCurveData.Reset();
