CurveEditor: Temporary fix for compiler error from CL7465722

diff --git a/Engine/Plugins/Editor/CurveEditorTools/Source/CurveEditorTools/Private/Tools/CurveEditorTransformTool.cpp b/Engine/Plugins/Editor/CurveEditorTools/Source/CurveEditorTools/Private/Tools/CurveEditorTransformTool.cpp
--- a/Engine/Plugins/Editor/CurveEditorTools/Source/CurveEditorTools/Private/Tools/CurveEditorTransformTool.cpp
+++ b/Engine/Plugins/Editor/CurveEditorTools/Source/CurveEditorTools/Private/Tools/CurveEditorTransformTool.cpp
@@ -490,8 +490,9 @@ void FCurveEditorTransformTool::OnDrag(const FPointerEvent& MouseEvent)
 				StartPosition.InputValue += DeltaInput;
 				StartPosition.OutputValue += DeltaOutput;
 
-				StartPosition.InputValue = View->IsTimeSnapEnabled() ? CurveEditor->GetSnapMetrics().SnapInputSeconds(StartPosition.InputValue) : StartPosition.InputValue;
-				StartPosition.OutputValue = View->IsValueSnapEnabled() ? CurveEditor->GetSnapMetrics().SnapOutput(StartPosition.OutputValue) : StartPosition.OutputValue;
+				// Temporarily disabled to fix build error caused by CL7465722
+				//StartPosition.InputValue = View->IsTimeSnapEnabled() ? CurveEditor->GetSnapMetrics().SnapInputSeconds(StartPosition.InputValue) : StartPosition.InputValue;
+				//StartPosition.OutputValue = View->IsValueSnapEnabled() ? CurveEditor->GetSnapMetrics().SnapOutput(StartPosition.OutputValue) : StartPosition.OutputValue;
 
 				NewKeyPositionScratch.Add(StartPosition);
 			}
