CurveEditor: Actual temporary fix for compiler error from CL7465722

diff --git a/Engine/Plugins/Editor/CurveEditorTools/Source/CurveEditorTools/Private/Tools/CurveEditorTransformTool.cpp b/Engine/Plugins/Editor/CurveEditorTools/Source/CurveEditorTools/Private/Tools/CurveEditorTransformTool.cpp
--- a/Engine/Plugins/Editor/CurveEditorTools/Source/CurveEditorTools/Private/Tools/CurveEditorTransformTool.cpp
+++ b/Engine/Plugins/Editor/CurveEditorTools/Source/CurveEditorTools/Private/Tools/CurveEditorTransformTool.cpp
@@ -602,11 +602,15 @@ void FCurveEditorTransformTool::OnDrag(const FPointerEvent& MouseEvent)
 				// This includes snapping, otherwise dragging on an edge can cause it to snap on the opposite axis.
 				if (bAffectsX)
 				{
-					StartPosition.InputValue = View->IsTimeSnapEnabled() ? CurveEditor->GetSnapMetrics().SnapInputSeconds(NewInput) : NewInput;
+					// Temporarily disabled to fix build error caused by CL7465722
+					//StartPosition.InputValue = View->IsTimeSnapEnabled() ? CurveEditor->GetSnapMetrics().SnapInputSeconds(NewInput) : NewInput;
+					StartPosition.InputValue = NewInput;
 				}
 				if (bAffectsY)
 				{
-					StartPosition.OutputValue = View->IsValueSnapEnabled() ? CurveEditor->GetSnapMetrics().SnapOutput(NewOutput) : NewOutput;
+					// Temporarily disabled to fix build error caused by CL7465722
+					//StartPosition.OutputValue = View->IsValueSnapEnabled() ? CurveEditor->GetSnapMetrics().SnapOutput(NewOutput) : NewOutput;
+					StartPosition.OutputValue = NewOutput;
 				}
 
 
