Curve Editor: Fix GetCurveSnapMetrics

diff --git a/Engine/Plugins/Editor/CurveEditorTools/Source/CurveEditorTools/Private/Tools/CurveEditorTransformTool.cpp b/Engine/Plugins/Editor/CurveEditorTools/Source/CurveEditorTools/Private/Tools/CurveEditorTransformTool.cpp
--- a/Engine/Plugins/Editor/CurveEditorTools/Source/CurveEditorTools/Private/Tools/CurveEditorTransformTool.cpp
+++ b/Engine/Plugins/Editor/CurveEditorTools/Source/CurveEditorTools/Private/Tools/CurveEditorTransformTool.cpp
@@ -490,9 +490,8 @@ void FCurveEditorTransformTool::OnDrag(const FPointerEvent& MouseEvent)
 				StartPosition.InputValue += DeltaInput;
 				StartPosition.OutputValue += DeltaOutput;
 
-				// Temporarily disabled to fix build error caused by CL7465722
-				//StartPosition.InputValue = View->IsTimeSnapEnabled() ? CurveEditor->GetSnapMetrics().SnapInputSeconds(StartPosition.InputValue) : StartPosition.InputValue;
-				//StartPosition.OutputValue = View->IsValueSnapEnabled() ? CurveEditor->GetSnapMetrics().SnapOutput(StartPosition.OutputValue) : StartPosition.OutputValue;
+				StartPosition.InputValue = View->IsTimeSnapEnabled() ? CurveEditor->GetCurveSnapMetrics(KeyData.CurveID).SnapInputSeconds(StartPosition.InputValue) : StartPosition.InputValue;
+				StartPosition.OutputValue = View->IsValueSnapEnabled() ? CurveEditor->GetCurveSnapMetrics(KeyData.CurveID).SnapOutput(StartPosition.OutputValue) : StartPosition.OutputValue;
 
 				NewKeyPositionScratch.Add(StartPosition);
 			}
@@ -602,15 +601,11 @@ void FCurveEditorTransformTool::OnDrag(const FPointerEvent& MouseEvent)
 				// This includes snapping, otherwise dragging on an edge can cause it to snap on the opposite axis.
 				if (bAffectsX)
 				{
-					// Temporarily disabled to fix build error caused by CL7465722
-					//StartPosition.InputValue = View->IsTimeSnapEnabled() ? CurveEditor->GetSnapMetrics().SnapInputSeconds(NewInput) : NewInput;
-					StartPosition.InputValue = NewInput;
+					StartPosition.InputValue = View->IsTimeSnapEnabled() ? CurveEditor->GetCurveSnapMetrics(KeyData.CurveID).SnapInputSeconds(NewInput) : NewInput;
 				}
 				if (bAffectsY)
 				{
-					// Temporarily disabled to fix build error caused by CL7465722
-					//StartPosition.OutputValue = View->IsValueSnapEnabled() ? CurveEditor->GetSnapMetrics().SnapOutput(NewOutput) : NewOutput;
-					StartPosition.OutputValue = NewOutput;
+					StartPosition.OutputValue = View->IsValueSnapEnabled() ? CurveEditor->GetCurveSnapMetrics(KeyData.CurveID).SnapOutput(NewOutput) : NewOutput;
 				}
 
 
