Curve Editor: add widget to show extension properties, add center scale anchor to the transform tool

diff --git a/Engine/Plugins/Editor/CurveEditorTools/Source/CurveEditorTools/Private/Tools/CurveEditorTransformTool.cpp b/Engine/Plugins/Editor/CurveEditorTools/Source/CurveEditorTools/Private/Tools/CurveEditorTransformTool.cpp
--- a/Engine/Plugins/Editor/CurveEditorTools/Source/CurveEditorTools/Private/Tools/CurveEditorTransformTool.cpp
+++ b/Engine/Plugins/Editor/CurveEditorTools/Source/CurveEditorTools/Private/Tools/CurveEditorTransformTool.cpp
@@ -22,16 +22,18 @@
 
 namespace CurveEditorTransformTool
 {
+	constexpr float ScaleCenterDiameter = 16.0f;
 	constexpr float EdgeAnchorWidth = 13.f;
+	constexpr float SoftSelecdtAnchorWidth = 20.f;
 	constexpr float EdgeHighlightAlpha = 0.15f;
 }
 
 void FCurveEditorTransformWidget::GetSidebarGeometry(const FGeometry& InWidgetGeometry, FGeometry& OutLeft, FGeometry& OutRight, FGeometry& OutTop, FGeometry& OutBottom) const
 {
-	FVector2D SidebarSize = FVector2D(CurveEditorTransformTool::EdgeAnchorWidth, InWidgetGeometry.GetLocalSize().Y - CurveEditorTransformTool::EdgeAnchorWidth);
-	FVector2D SidebarSizeOffset = FVector2D(CurveEditorTransformTool::EdgeAnchorWidth / 2.f, 0.f);
-	FVector2D TopbarSize = FVector2D(InWidgetGeometry.GetLocalSize().X - CurveEditorTransformTool::EdgeAnchorWidth, CurveEditorTransformTool::EdgeAnchorWidth);
-	FVector2D TopbarSizeOffset = FVector2D(0.f, CurveEditorTransformTool::EdgeAnchorWidth / 2.f);
+	const FVector2D SidebarSize = FVector2D(CurveEditorTransformTool::EdgeAnchorWidth, InWidgetGeometry.GetLocalSize().Y - CurveEditorTransformTool::EdgeAnchorWidth);
+	const FVector2D SidebarSizeOffset = FVector2D(CurveEditorTransformTool::EdgeAnchorWidth / 2.f, 0.f);
+	const FVector2D TopbarSize = FVector2D(InWidgetGeometry.GetLocalSize().X - CurveEditorTransformTool::EdgeAnchorWidth, CurveEditorTransformTool::EdgeAnchorWidth);
+	const FVector2D TopbarSizeOffset = FVector2D(0.f, CurveEditorTransformTool::EdgeAnchorWidth / 2.f);
 
 	OutLeft = InWidgetGeometry.MakeChild(SidebarSize, FSlateLayoutTransform(FVector2D(0, CurveEditorTransformTool::EdgeAnchorWidth / 2.f) - SidebarSizeOffset));
 	OutRight = InWidgetGeometry.MakeChild(SidebarSize, FSlateLayoutTransform(FVector2D(InWidgetGeometry.GetLocalSize().X, CurveEditorTransformTool::EdgeAnchorWidth / 2.f) - SidebarSizeOffset));
@@ -41,8 +43,8 @@ void FCurveEditorTransformWidget::GetSidebarGeometry(const FGeometry& InWidgetGe
 
 void FCurveEditorTransformWidget::GetCornerGeometry(const FGeometry& InWidgetGeometry, FGeometry& OutTopLeft, FGeometry& OutTopRight, FGeometry& OutBottomLeft, FGeometry& OutBottomRight) const
 {
-	FVector2D CornerSize = FVector2D(CurveEditorTransformTool::EdgeAnchorWidth, CurveEditorTransformTool::EdgeAnchorWidth);
-	FVector2D HalfSizeOffset = FVector2D(CornerSize / 2.f);
+	const FVector2D CornerSize = FVector2D(CurveEditorTransformTool::EdgeAnchorWidth, CurveEditorTransformTool::EdgeAnchorWidth);
+	const FVector2D HalfSizeOffset = FVector2D(CornerSize / 2.f);
 
 	FSlateLayoutTransform TopLeftPosition = FSlateLayoutTransform(FVector2D(0, 0) - HalfSizeOffset);
 	FSlateLayoutTransform TopRightPosition = FSlateLayoutTransform(FVector2D(InWidgetGeometry.GetLocalSize().X, 0) - HalfSizeOffset);
@@ -57,18 +59,82 @@ void FCurveEditorTransformWidget::GetCornerGeometry(const FGeometry& InWidgetGeo
 
 void FCurveEditorTransformWidget::GetCenterGeometry(const FGeometry& InWidgetGeometry, FGeometry& OutCenter) const
 {
-	FVector2D CenterSize = InWidgetGeometry.GetLocalSize() - FVector2D(CurveEditorTransformTool::EdgeAnchorWidth, CurveEditorTransformTool::EdgeAnchorWidth);
-	FVector2D CenterOffset = FVector2D(CurveEditorTransformTool::EdgeAnchorWidth / 2.f, CurveEditorTransformTool::EdgeAnchorWidth / 2.f);
+	const FVector2D CenterSize = InWidgetGeometry.GetLocalSize() - FVector2D(CurveEditorTransformTool::EdgeAnchorWidth, CurveEditorTransformTool::EdgeAnchorWidth);
+	const FVector2D CenterOffset = FVector2D(CurveEditorTransformTool::EdgeAnchorWidth / 2.f, CurveEditorTransformTool::EdgeAnchorWidth / 2.f);
 
 	OutCenter = InWidgetGeometry.MakeChild(CenterSize, FSlateLayoutTransform(CenterOffset));
 }
 
-ECurveEditorAnchorFlags FCurveEditorTransformWidget::GetAnchorFlagsForMousePosition(const FGeometry& InWidgetGeometry, const FVector2D& InMouseScreenPosition) const
+void FCurveEditorTransformWidget::GetFalloffGeometry(const FGeometry& InWidgetGeometry, float FalloffHeight, float FalloffWidth, FGeometry& OutTopLeft, FGeometry& OutTopRight, FGeometry& OutLeft, FGeometry& OutRight) const
+{
+	const FVector2D CornerSize = FVector2D(CurveEditorTransformTool::SoftSelecdtAnchorWidth, CurveEditorTransformTool::SoftSelecdtAnchorWidth);
+	const FVector2D HalfSizeOffset = FVector2D(CornerSize / 2.f);
+	const float HeightDiff = InWidgetGeometry.GetLocalSize().Y * FalloffHeight;
+	const float WidthDiff = InWidgetGeometry.GetLocalSize().X * 0.5f * FalloffWidth;
+	FSlateLayoutTransform TopLeftPosition = FSlateLayoutTransform(FVector2D(InWidgetGeometry.GetLocalSize().X * 0.5f - WidthDiff, 0) - HalfSizeOffset);
+	FSlateLayoutTransform TopRightPosition = FSlateLayoutTransform(FVector2D(InWidgetGeometry.GetLocalSize().X * 0.5f + WidthDiff, 0) - HalfSizeOffset);
+	FSlateLayoutTransform BottomLeftPosition = FSlateLayoutTransform(FVector2D(0, InWidgetGeometry.GetLocalSize().Y - HeightDiff) - HalfSizeOffset);
+	FSlateLayoutTransform BottomRightPosition = FSlateLayoutTransform(FVector2D(InWidgetGeometry.GetLocalSize().X, InWidgetGeometry.GetLocalSize().Y - HeightDiff) - HalfSizeOffset);
+
+	OutTopLeft = InWidgetGeometry.MakeChild(CornerSize, TopLeftPosition);
+	OutTopRight = InWidgetGeometry.MakeChild(CornerSize, TopRightPosition);
+	OutLeft = InWidgetGeometry.MakeChild(CornerSize, BottomLeftPosition);
+	OutRight = InWidgetGeometry.MakeChild(CornerSize, BottomRightPosition);
+}
+
+void FCurveEditorTransformWidget::GetScaleCenterGeometry(const FGeometry& InWidgetGeometry, FVector2D ScaleCenter, FGeometry& OutScaleCenterGeometry) const
+{
+	const FVector2D CenterSize = FVector2D(CurveEditorTransformTool::ScaleCenterDiameter * 2.f, CurveEditorTransformTool::ScaleCenterDiameter * 2.f);
+	const FVector2D CenterOffset = (InWidgetGeometry.GetLocalSize() * .5f) + (BoundsSize * (ScaleCenter - 0.5f)) - (CenterSize * .5f);
+	OutScaleCenterGeometry = InWidgetGeometry.MakeChild(CenterSize, FSlateLayoutTransform(CenterOffset));
+}
+
+ECurveEditorAnchorFlags FCurveEditorTransformWidget::GetAnchorFlagsForMousePosition(const FGeometry& InWidgetGeometry, float  FalloffHeight, float FalloffWidth, const FVector2D& RelativeCenterScale, const FVector2D& InMouseScreenPosition) const
 {
 	// We store a geometry to represent each different region, updated on Tick. We check if the mouse
 	// overlaps a region and update the selection anchors depending on which region you're hovering in.
 	ECurveEditorAnchorFlags  OutFlags = ECurveEditorAnchorFlags::None;
 
+	// If holding ctrl we test for falloff first and then bail if getting hit 
+	if (FSlateApplication::Get().GetModifierKeys().IsControlDown())
+	{
+		FGeometry TopLeftFalloffGeometry, TopRightFalloffGeometry, LeftFalloffGeometry, RightFalloffGeometry;
+		GetFalloffGeometry(InWidgetGeometry, FalloffHeight, FalloffWidth, TopLeftFalloffGeometry, TopRightFalloffGeometry, LeftFalloffGeometry, RightFalloffGeometry);
+
+		if (TopLeftFalloffGeometry.IsUnderLocation(InMouseScreenPosition))
+		{
+			OutFlags |= ECurveEditorAnchorFlags::FalloffTopLeft;
+		}
+		if (TopRightFalloffGeometry.IsUnderLocation(InMouseScreenPosition))
+		{
+			OutFlags |= ECurveEditorAnchorFlags::FalloffTopRight;
+		}
+		if (LeftFalloffGeometry.IsUnderLocation(InMouseScreenPosition))
+		{
+			OutFlags |= ECurveEditorAnchorFlags::FalloffLeft;
+		}
+		if (RightFalloffGeometry.IsUnderLocation(InMouseScreenPosition))
+		{
+			OutFlags |= ECurveEditorAnchorFlags::FalloffRight;
+		}
+		if (OutFlags != ECurveEditorAnchorFlags::None)
+		{
+			return OutFlags;
+		}
+	}
+
+	if (FSlateApplication::Get().GetModifierKeys().IsAltDown())
+	{
+		FGeometry CenterScaleGeometry;
+		GetScaleCenterGeometry(InWidgetGeometry, RelativeCenterScale, CenterScaleGeometry);
+
+		if (CenterScaleGeometry.IsUnderLocation(InMouseScreenPosition))
+		{
+			OutFlags |= ECurveEditorAnchorFlags::CenterScale;
+			return OutFlags;
+		}
+	}
+
 	FGeometry LeftSidebarGeometry, RightSidebarGeometry, TopSidebarGeometry, BottomSidebarGeometry;
 	FGeometry TopLeftCornerGeometry, TopRightCornerGeometry, BottomLeftCornerGeometry, BottomRightCornerGeometry;
 
@@ -123,29 +189,15 @@ ECurveEditorAnchorFlags FCurveEditorTransformWidget::GetAnchorFlagsForMousePosit
 	return OutFlags;
 }
 
-
 void FCurveEditorTransformTool::OnToolActivated()
 {
-	// UE_LOG(LogTemp, Log, TEXT("TransformTool Activated."));
-	TSharedPtr<FCurveEditor> CurveEditor = WeakCurveEditor.Pin();
-	if (CurveEditor.IsValid())
-	{
-		// CurveEditor->GetSelection().SelectionChangedDelegate += OnSelectionChanged;
-	}
-
-	// Manually fire an OnSelectionChanged event so we go look up the selection and build our bounding box
-	// OnSelectionChanged();
+	// set the new tool option values
+	UpdateToolOptions();
 }
 
 void FCurveEditorTransformTool::OnToolDeactivated()
 {
-	// UE_LOG(LogTemp, Log, TEXT("TransformTool Deactivated."));
-
-	TSharedPtr<FCurveEditor> CurveEditor = WeakCurveEditor.Pin();
-	if (CurveEditor.IsValid())
-	{
-		// CurveEditor->GetSelection().SelectionChangedDelegate -= OnSelectionChanged;
-	}
+	
 }
 
 void FCurveEditorTransformTool::UpdateMarqueeBoundingBox()
@@ -193,8 +245,8 @@ void FCurveEditorTransformTool::UpdateMarqueeBoundingBox()
 
 		for (int32 i = 0; i < KeyPositions.Num(); i++)
 		{
-			FVector2D ViewSpaceLocation = FVector2D(CurveSpace.SecondsToScreen(KeyPositions[i].InputValue), CurveSpace.ValueToScreen(KeyPositions[i].OutputValue));
-			FVector2D PanelSpaceLocation = InnerToOuterTransform.TransformPoint(ViewSpaceLocation);
+			const FVector2D ViewSpaceLocation = FVector2D(CurveSpace.SecondsToScreen(KeyPositions[i].InputValue), CurveSpace.ValueToScreen(KeyPositions[i].OutputValue));
+			const FVector2D PanelSpaceLocation = InnerToOuterTransform.TransformPoint(ViewSpaceLocation);
 
 			if (!MinValue.IsSet())
 			{
@@ -216,6 +268,9 @@ void FCurveEditorTransformTool::UpdateMarqueeBoundingBox()
 		FVector2D MarqueeSize = MaxValue.GetValue() - MinValue.GetValue();
 		FVector2D Offset = FVector2D::ZeroVector;
 		
+		TransformWidget.BoundsSize = MarqueeSize;
+		TransformWidget.BoundsPosition = MinValue.GetValue();
+
 		// Enforce a minimum size for single time/value selections.
 		if (MarqueeSize.X < 8.f)
 		{
@@ -240,24 +295,45 @@ void FCurveEditorTransformTool::UpdateMarqueeBoundingBox()
 		TransformWidget.Size = FVector2D::ZeroVector;
 		TransformWidget.Position = FVector2D::ZeroVector;
 	}
+}
 
-	// UE_LOG(LogTemp, Log, TEXT("Size: %s Position: %s"), *TransformWidget.Size.ToString(), *TransformWidget.Position.ToString());
+void FCurveEditorTransformTool::UpdateToolOptions()
+{
+	TSharedPtr<FCurveEditor> CurveEditor = WeakCurveEditor.Pin();
+	if (CurveEditor.IsValid())
+	{
+		TArray<FCurveModelID> CurveKeys;
+		CurveEditor->GetSelection().GetAll().GetKeys(CurveKeys);
+		if (CurveKeys.Num() < 1)
+		{
+			return;
+		}
+		const SCurveEditorView* View = CurveEditor->FindFirstInteractiveView(CurveKeys[0]);
+		if (!View)
+		{
+			return;
+		}
+		FCurveEditorScreenSpace CurveSpace = View->GetCurveSpace(CurveKeys[0]);
+
+		ToolOptions.LeftBound = CurveSpace.ScreenToSeconds(TransformWidget.BoundsPosition.X);
+		ToolOptions.UpperBound = CurveSpace.ScreenToValue(TransformWidget.BoundsPosition.Y);
+		ToolOptions.RightBound = CurveSpace.ScreenToSeconds(TransformWidget.BoundsPosition.X + TransformWidget.BoundsSize.X);
+		ToolOptions.LowerBound = CurveSpace.ScreenToValue(TransformWidget.BoundsPosition.Y + TransformWidget.BoundsSize.Y);
+	}
 }
 
 FReply FCurveEditorTransformTool::OnMouseButtonDown(TSharedRef<SWidget> OwningWidget, const FGeometry& MyGeometry, const FPointerEvent& MouseEvent)
 {
-	// UE_LOG(LogTemp, Log, TEXT("FCurveEditorTransformTool::OnMouseButtonDown"));
 	DelayedDrag.Reset();
 	if (MouseEvent.GetEffectingButton() == EKeys::LeftMouseButton)
 	{
 		FGeometry WidgetGeo = TransformWidget.MakeGeometry(MyGeometry);
-		ECurveEditorAnchorFlags HitWidgetFlags = TransformWidget.GetAnchorFlagsForMousePosition(WidgetGeo, MouseEvent.GetScreenSpacePosition());
+		ECurveEditorAnchorFlags HitWidgetFlags = TransformWidget.GetAnchorFlagsForMousePosition(WidgetGeo, FalloffHeight, FalloffWidth, RelativeScaleCenter, MouseEvent.GetScreenSpacePosition());
 		if (HitWidgetFlags != ECurveEditorAnchorFlags::None)
 		{
-			// TransformWidget.SelectedAnchorFlags = HitWidgetFlags;
-			
 			// Start a Delayed Drag
 			DelayedDrag = FDelayedDrag(MyGeometry.AbsoluteToLocal(MouseEvent.GetScreenSpacePosition()), MouseEvent.GetEffectingButton());
+			
 			return FReply::Handled();
 		}
 	}
@@ -270,7 +346,7 @@ FReply FCurveEditorTransformTool::OnMouseMove(TSharedRef<SWidget> OwningWidget,
 	if (!DelayedDrag.IsSet())
 	{
 		FGeometry WidgetGeo = TransformWidget.MakeGeometry(MyGeometry);
-		ECurveEditorAnchorFlags HitWidgetFlags = TransformWidget.GetAnchorFlagsForMousePosition(WidgetGeo, MouseEvent.GetScreenSpacePosition());
+		ECurveEditorAnchorFlags HitWidgetFlags = TransformWidget.GetAnchorFlagsForMousePosition(WidgetGeo, FalloffHeight, FalloffWidth, RelativeScaleCenter, MouseEvent.GetScreenSpacePosition());
 		TransformWidget.SelectedAnchorFlags = HitWidgetFlags;
 	}
 
@@ -280,7 +356,8 @@ FReply FCurveEditorTransformTool::OnMouseMove(TSharedRef<SWidget> OwningWidget,
 
 		if (DelayedDrag->IsDragging())
 		{
-			OnDrag(MouseEvent);
+			const FVector2D LocalMousePosition = MyGeometry.AbsoluteToLocal(MouseEvent.GetScreenSpacePosition());
+			OnDrag(MouseEvent, LocalMousePosition);
 		}
 		else if (DelayedDrag->AttemptDragStart(MouseEvent))
 		{
@@ -299,7 +376,6 @@ FReply FCurveEditorTransformTool::OnMouseMove(TSharedRef<SWidget> OwningWidget,
 
 FReply FCurveEditorTransformTool::OnMouseButtonUp(TSharedRef<SWidget> OwningWidget, const FGeometry& MyGeometry, const FPointerEvent& MouseEvent)
 {
-	// UE_LOG(LogTemp, Log, TEXT("FCurveEditorTransformTool::OnMouseButtonUp"));
 	FReply Reply = FReply::Handled();
 	if (DelayedDrag.IsSet())
 	{
@@ -310,7 +386,6 @@ FReply FCurveEditorTransformTool::OnMouseButtonUp(TSharedRef<SWidget> OwningWidg
 			// Only return handled if we actually started a drag
 			Reply.ReleaseMouseCapture();
 		}
-
 		DelayedDrag.Reset();
 		return Reply;
 	}
@@ -320,12 +395,80 @@ FReply FCurveEditorTransformTool::OnMouseButtonUp(TSharedRef<SWidget> OwningWidg
 
 void FCurveEditorTransformTool::OnFocusLost(const FFocusEvent& InFocusEvent)
 {
-	// UE_LOG(LogTemp, Log, TEXT("FCurveEditorTimeTool::OnFocusLost."));
 	// We need to end our drag if we lose Window focus to close the transaction, otherwise alt-tabbing while dragging
 	// can cause a transaction to get stuck open.
 	StopDragIfPossible();
 }
 
+void FCurveEditorTransformTool::OnToolOptionsUpdated(const FPropertyChangedEvent& PropertyChangedEvent)
+{
+	TSharedPtr<FCurveEditor> CurveEditor = WeakCurveEditor.Pin();
+	if (!CurveEditor)
+	{
+		return;
+	}
+
+	// call OnDragStart() to update the initial key positions in KeysByCurve for ScaleFrom() to scale
+	OnDragStart();
+
+	TArray<FCurveModelID> CurveKeys;
+	CurveEditor->GetSelection().GetAll().GetKeys(CurveKeys);
+	if (CurveKeys.Num() < 1)
+	{
+		return;
+	}
+
+	// grab the view for the first curve, if there are multiple curves it is assumed they all have the same curve space
+	const SCurveEditorView* View = CurveEditor->FindFirstInteractiveView(CurveKeys[0]);
+	if (!View)
+	{
+		return;
+	}
+
+	FCurveEditorScreenSpace CurveSpace = View->GetCurveSpace(CurveKeys[0]);
+
+	FVector2D ScaleCenter = FVector2D(0.5f, 0.5f);
+	FVector2D ScaleDelta = FVector2D(0.0f, 0.0f);
+	bool bAffectsX = true, bAffectsY = true;
+
+	if (PropertyChangedEvent.GetPropertyName().IsEqual(GET_MEMBER_NAME_CHECKED(FTransformToolOptions, LeftBound)))
+	{
+		bAffectsY = false;
+		ScaleCenter.X = 1.0f;
+		ScaleDelta.X = -(CurveSpace.SecondsToScreen(ToolOptions.LeftBound) - TransformWidget.Position.X);
+	}
+	else if (PropertyChangedEvent.GetPropertyName().IsEqual(GET_MEMBER_NAME_CHECKED(FTransformToolOptions, UpperBound)))
+	{
+		bAffectsX = false;
+		ScaleCenter.Y = 1.0f;
+		ScaleDelta.Y = -(CurveSpace.ValueToScreen(ToolOptions.UpperBound) - TransformWidget.Position.Y);
+		float diff = ToolOptions.UpperBound - CurveSpace.ScreenToValue(TransformWidget.Position.Y);
+		UE_LOG(LogTemp, Log, TEXT("diff: %f"), diff);
+	}
+	else if (PropertyChangedEvent.GetPropertyName().IsEqual(GET_MEMBER_NAME_CHECKED(FTransformToolOptions, RightBound)))
+	{
+		bAffectsY = false;
+		ScaleCenter.X = 0.0f;
+		ScaleDelta.X = CurveSpace.SecondsToScreen(ToolOptions.RightBound) - (TransformWidget.Position.X + TransformWidget.Size.X);
+	}
+	else if (PropertyChangedEvent.GetPropertyName().IsEqual(GET_MEMBER_NAME_CHECKED(FTransformToolOptions, LowerBound)))
+	{
+		bAffectsX = false;
+		ScaleCenter.Y = 0.0f;
+		ScaleDelta.Y = CurveSpace.ValueToScreen(ToolOptions.LowerBound) - (TransformWidget.Position.Y + TransformWidget.Size.Y);
+	}
+
+	
+	const FVector2D PanelSpaceCenter = TransformWidget.Position + (TransformWidget.Size * ScaleCenter);
+	const FVector2D ChangeAmount = (ScaleDelta / TransformWidget.Size);
+
+ 	ScaleFrom(PanelSpaceCenter, ChangeAmount, false, bAffectsX, bAffectsY);
+
+	UpdateMarqueeBoundingBox();
+
+	PrevState.Positon = TransformWidget.Position;
+	PrevState.Size = TransformWidget.Size;
+}
 
 void FCurveEditorTransformTool::OnPaint(const FPaintArgs& Args, const FGeometry& AllottedGeometry, const FSlateRect& MyCullingRect, FSlateWindowElementList& OutDrawElements, int32 PaintOnLayerId, const FWidgetStyle& InWidgetStyle, bool bParentEnabled) const
 {
@@ -334,7 +477,7 @@ void FCurveEditorTransformTool::OnPaint(const FPaintArgs& Args, const FGeometry&
 	DrawMarqueeWidget(TransformWidget, Args, WidgetGeo, MyCullingRect, OutDrawElements, PaintOnLayerId, InWidgetStyle, bParentEnabled);
 }
 
-void FCurveEditorTransformTool::DrawMarqueeWidget(const FCurveEditorTransformWidget& InTransformWidget, const FPaintArgs& Args, const FGeometry& AllottedGeometry, const FSlateRect& MyCullingRect, FSlateWindowElementList& OutDrawElements, int32 PaintOnLayerId, const FWidgetStyle& InWidgetStyle, bool bParentEnabled) const
+void FCurveEditorTransformTool::DrawMarqueeWidget(const FCurveEditorTransformWidget& InTransformWidget, const FPaintArgs& InArgs, const FGeometry& InAllottedGeometry, const FSlateRect& InMyCullingRect, FSlateWindowElementList& OutDrawElements, const int32 InPaintOnLayerId, const FWidgetStyle& InWidgetStyle, const bool bInParentEnabled) const
 {
 	if (!InTransformWidget.Visible)
 	{
@@ -345,20 +488,12 @@ void FCurveEditorTransformTool::DrawMarqueeWidget(const FCurveEditorTransformWid
 	{
 		FLinearColor CenterHighlightColor = (InTransformWidget.SelectedAnchorFlags == ECurveEditorAnchorFlags::Center) ? FLinearColor::White.CopyWithNewOpacity(CurveEditorTransformTool::EdgeHighlightAlpha) : FLinearColor::Transparent;
 		FGeometry CenterGeometry;
-		InTransformWidget.GetCenterGeometry(AllottedGeometry, CenterGeometry);
+		InTransformWidget.GetCenterGeometry(InAllottedGeometry, CenterGeometry);
 
-		FSlateDrawElement::MakeBox(OutDrawElements, PaintOnLayerId, CenterGeometry.ToPaintGeometry(), FEditorStyle::GetBrush(TEXT("WhiteBrush")), ESlateDrawEffect::None, CenterHighlightColor);
-		FSlateDrawElement::MakeBox(OutDrawElements, PaintOnLayerId, AllottedGeometry.ToPaintGeometry(), FEditorStyle::GetBrush(TEXT("MarqueeSelection")));
+		FSlateDrawElement::MakeBox(OutDrawElements, InPaintOnLayerId, CenterGeometry.ToPaintGeometry(), FEditorStyle::GetBrush(TEXT("WhiteBrush")), ESlateDrawEffect::None, CenterHighlightColor);
+		FSlateDrawElement::MakeBox(OutDrawElements, InPaintOnLayerId, InAllottedGeometry.ToPaintGeometry(), FEditorStyle::GetBrush(TEXT("MarqueeSelection")));
 	}
 
-	// Draw the thing that darkens the background
-	{
-		// This is comprised of four boxes that exclude the Marquee dimensions in the middle.
-	}
-
-	// Draw inner gradient which indicates tool falloff
-
-
 	// Draw edge highlight regions on mouse hover
 	{
 		FLinearColor LeftEdgeHighlightColor =	(InTransformWidget.SelectedAnchorFlags == ECurveEditorAnchorFlags::Left)	? FLinearColor::White.CopyWithNewOpacity(CurveEditorTransformTool::EdgeHighlightAlpha) : FLinearColor::Transparent;
@@ -367,22 +502,90 @@ void FCurveEditorTransformTool::DrawMarqueeWidget(const FCurveEditorTransformWid
 		FLinearColor BottomEdgeHighlightColor =	(InTransformWidget.SelectedAnchorFlags == ECurveEditorAnchorFlags::Bottom)	? FLinearColor::White.CopyWithNewOpacity(CurveEditorTransformTool::EdgeHighlightAlpha) : FLinearColor::Transparent;
 
 		FGeometry LeftSidebarGeometry, RightSidebarGeometry, TopSidebarGeometry, BottomSidebarGeometry;
-		InTransformWidget.GetSidebarGeometry(AllottedGeometry, LeftSidebarGeometry, RightSidebarGeometry, TopSidebarGeometry, BottomSidebarGeometry);
+		InTransformWidget.GetSidebarGeometry(InAllottedGeometry, LeftSidebarGeometry, RightSidebarGeometry, TopSidebarGeometry, BottomSidebarGeometry);
 
 		// Left Edge
-		FSlateDrawElement::MakeBox(OutDrawElements, PaintOnLayerId, LeftSidebarGeometry.ToPaintGeometry(), FEditorStyle::GetBrush(TEXT("WhiteBrush")), ESlateDrawEffect::None, LeftEdgeHighlightColor);
+		FSlateDrawElement::MakeBox(OutDrawElements, InPaintOnLayerId, LeftSidebarGeometry.ToPaintGeometry(), FEditorStyle::GetBrush(TEXT("WhiteBrush")), ESlateDrawEffect::None, LeftEdgeHighlightColor);
 		// Right Edge
-		FSlateDrawElement::MakeBox(OutDrawElements, PaintOnLayerId, RightSidebarGeometry.ToPaintGeometry(), FEditorStyle::GetBrush(TEXT("WhiteBrush")), ESlateDrawEffect::None, RightEdgeHighlightColor);
+		FSlateDrawElement::MakeBox(OutDrawElements, InPaintOnLayerId, RightSidebarGeometry.ToPaintGeometry(), FEditorStyle::GetBrush(TEXT("WhiteBrush")), ESlateDrawEffect::None, RightEdgeHighlightColor);
 		// Top Edge
-		FSlateDrawElement::MakeBox(OutDrawElements, PaintOnLayerId, TopSidebarGeometry.ToPaintGeometry(), FEditorStyle::GetBrush(TEXT("WhiteBrush")), ESlateDrawEffect::None, TopEdgeHighlightColor);
+		FSlateDrawElement::MakeBox(OutDrawElements, InPaintOnLayerId, TopSidebarGeometry.ToPaintGeometry(), FEditorStyle::GetBrush(TEXT("WhiteBrush")), ESlateDrawEffect::None, TopEdgeHighlightColor);
 		// Bottom Edge
-		FSlateDrawElement::MakeBox(OutDrawElements, PaintOnLayerId, BottomSidebarGeometry.ToPaintGeometry(), FEditorStyle::GetBrush(TEXT("WhiteBrush")), ESlateDrawEffect::None, BottomEdgeHighlightColor);
+		FSlateDrawElement::MakeBox(OutDrawElements, InPaintOnLayerId, BottomSidebarGeometry.ToPaintGeometry(), FEditorStyle::GetBrush(TEXT("WhiteBrush")), ESlateDrawEffect::None, BottomEdgeHighlightColor);
+	}
+	// Draw falloff if on
+	if (FSlateApplication::Get().GetModifierKeys().IsControlDown())
+	{
+		FGeometry TopLeftFalloffGeometry, TopRightFalloffGeometry, LeftFalloffGeometry, RightFalloffGeometry;
+		InTransformWidget.GetFalloffGeometry(InAllottedGeometry, FalloffHeight, FalloffWidth, TopLeftFalloffGeometry, TopRightFalloffGeometry, LeftFalloffGeometry, RightFalloffGeometry);
+
+		// fixed logic to set bools correctly
+		const bool bTopLeft = InTransformWidget.SelectedAnchorFlags == ECurveEditorAnchorFlags::FalloffTopLeft || InTransformWidget.SelectedAnchorFlags == ECurveEditorAnchorFlags::FalloffTopRight;
+		const bool bTopRight = InTransformWidget.SelectedAnchorFlags == ECurveEditorAnchorFlags::FalloffTopLeft || InTransformWidget.SelectedAnchorFlags == ECurveEditorAnchorFlags::FalloffTopRight;
+		const bool bBottomLeft = InTransformWidget.SelectedAnchorFlags == ECurveEditorAnchorFlags::FalloffLeft || InTransformWidget.SelectedAnchorFlags == ECurveEditorAnchorFlags::FalloffRight;
+		const bool bBottomRight = InTransformWidget.SelectedAnchorFlags == ECurveEditorAnchorFlags::FalloffLeft || InTransformWidget.SelectedAnchorFlags == ECurveEditorAnchorFlags::FalloffRight;
+
+		FLinearColor TopLeftHighlightColor = bTopLeft ? FLinearColor::White.CopyWithNewOpacity(CurveEditorTransformTool::EdgeHighlightAlpha) : FLinearColor::Transparent;
+		FLinearColor TopRightHighlightColor = bTopRight ? FLinearColor::White.CopyWithNewOpacity(CurveEditorTransformTool::EdgeHighlightAlpha) : FLinearColor::Transparent;
+		FLinearColor BottomLeftHighlightColor = bBottomLeft ? FLinearColor::White.CopyWithNewOpacity(CurveEditorTransformTool::EdgeHighlightAlpha) : FLinearColor::Transparent;
+		FLinearColor BottomRightHighlightColor = bBottomRight ? FLinearColor::White.CopyWithNewOpacity(CurveEditorTransformTool::EdgeHighlightAlpha) : FLinearColor::Transparent;
+
+		const float Rotate = 180.0f;
+		// Top Left (Highlight, Corner Icon)
+		FSlateDrawElement::MakeRotatedBox(OutDrawElements, InPaintOnLayerId, TopLeftFalloffGeometry.ToPaintGeometry(), FEditorStyle::GetBrush(TEXT("WhiteBrush")), ESlateDrawEffect::None, Rotate,
+			TOptional<FVector2D>(), FSlateDrawElement::RelativeToElement, TopLeftHighlightColor);
+		FSlateDrawElement::MakeRotatedBox(OutDrawElements, InPaintOnLayerId, TopLeftFalloffGeometry.ToPaintGeometry(), FEditorStyle::GetBrush(TEXT("MarqueeSelection")), ESlateDrawEffect::None, Rotate);
+		// Top Right										 
+		FSlateDrawElement::MakeRotatedBox(OutDrawElements, InPaintOnLayerId, TopRightFalloffGeometry.ToPaintGeometry(), FEditorStyle::GetBrush(TEXT("WhiteBrush")), ESlateDrawEffect::None, Rotate,
+			TOptional<FVector2D>(), FSlateDrawElement::RelativeToElement, TopRightHighlightColor);
+		FSlateDrawElement::MakeRotatedBox(OutDrawElements, InPaintOnLayerId, TopRightFalloffGeometry.ToPaintGeometry(), FEditorStyle::GetBrush(TEXT("MarqueeSelection")), ESlateDrawEffect::None, Rotate);
+		// Bottom Left										 
+		FSlateDrawElement::MakeRotatedBox(OutDrawElements, InPaintOnLayerId, LeftFalloffGeometry.ToPaintGeometry(), FEditorStyle::GetBrush(TEXT("WhiteBrush")), ESlateDrawEffect::None, Rotate,
+			TOptional<FVector2D>(), FSlateDrawElement::RelativeToElement, BottomLeftHighlightColor);
+		FSlateDrawElement::MakeRotatedBox(OutDrawElements, InPaintOnLayerId, LeftFalloffGeometry.ToPaintGeometry(), FEditorStyle::GetBrush(TEXT("MarqueeSelection")), ESlateDrawEffect::None, Rotate);
+		// Bottom Right										 
+		FSlateDrawElement::MakeRotatedBox(OutDrawElements, InPaintOnLayerId, RightFalloffGeometry.ToPaintGeometry(), FEditorStyle::GetBrush(TEXT("WhiteBrush")), ESlateDrawEffect::None, Rotate,
+			TOptional<FVector2D>(), FSlateDrawElement::RelativeToElement, BottomRightHighlightColor);
+		FSlateDrawElement::MakeRotatedBox(OutDrawElements, InPaintOnLayerId, RightFalloffGeometry.ToPaintGeometry(), FEditorStyle::GetBrush(TEXT("MarqueeSelection")), ESlateDrawEffect::None, Rotate);
+	}
+	// draw center marker if on
+	if (FSlateApplication::Get().GetModifierKeys().IsAltDown())
+	{
+		FGeometry ScaleCenterGeometry;
+		InTransformWidget.GetScaleCenterGeometry(InAllottedGeometry, RelativeScaleCenter, ScaleCenterGeometry);
+		
+		const bool bScaleCenter = InTransformWidget.SelectedAnchorFlags == ECurveEditorAnchorFlags::CenterScale;
+
+		FLinearColor ScaleCenterHighlightColor = bScaleCenter ? FLinearColor::White.CopyWithNewOpacity(CurveEditorTransformTool::EdgeHighlightAlpha) : FLinearColor::Transparent;
+		const float HighlightThickness = 7.f;
+
+		// draw center scale icon lines
+		FVector2D CircleOffset = FVector2D(CurveEditorTransformTool::ScaleCenterDiameter, CurveEditorTransformTool::ScaleCenterDiameter);
+		FSlateDrawElement::MakeSpline(OutDrawElements, InPaintOnLayerId, ScaleCenterGeometry.ToPaintGeometry(), FVector2D(-CurveEditorTransformTool::ScaleCenterDiameter, 0.f) + CircleOffset, FVector2D(CurveEditorTransformTool::ScaleCenterDiameter, 0.f), 
+			FVector2D(0.f, -CurveEditorTransformTool::ScaleCenterDiameter) + CircleOffset, FVector2D(0.f, -CurveEditorTransformTool::ScaleCenterDiameter), 1.f);
+		FSlateDrawElement::MakeSpline(OutDrawElements, InPaintOnLayerId, ScaleCenterGeometry.ToPaintGeometry(), FVector2D(CurveEditorTransformTool::ScaleCenterDiameter, 0.f) + CircleOffset, FVector2D(-CurveEditorTransformTool::ScaleCenterDiameter, 0.f),
+			FVector2D(0.f, CurveEditorTransformTool::ScaleCenterDiameter) + CircleOffset, FVector2D(0.f, CurveEditorTransformTool::ScaleCenterDiameter), 1.f);
+		FSlateDrawElement::MakeSpline(OutDrawElements, InPaintOnLayerId, ScaleCenterGeometry.ToPaintGeometry(), FVector2D(0.f, CurveEditorTransformTool::ScaleCenterDiameter) + CircleOffset, FVector2D(0.f, -CurveEditorTransformTool::ScaleCenterDiameter),
+			FVector2D(-CurveEditorTransformTool::ScaleCenterDiameter, 0.f) + CircleOffset, FVector2D(-CurveEditorTransformTool::ScaleCenterDiameter, 0.f), 1.f);
+		FSlateDrawElement::MakeSpline(OutDrawElements, InPaintOnLayerId, ScaleCenterGeometry.ToPaintGeometry(), FVector2D(0.f, -CurveEditorTransformTool::ScaleCenterDiameter) + CircleOffset, FVector2D(0.f, CurveEditorTransformTool::ScaleCenterDiameter),
+			FVector2D(CurveEditorTransformTool::ScaleCenterDiameter, 0.f) + CircleOffset, FVector2D(CurveEditorTransformTool::ScaleCenterDiameter, 0.f), 1.f);
+
+		// draw center scale icon highlight
+		FSlateDrawElement::MakeSpline(OutDrawElements, InPaintOnLayerId, ScaleCenterGeometry.ToPaintGeometry(), FVector2D(-CurveEditorTransformTool::ScaleCenterDiameter, 0.f) + CircleOffset, FVector2D(CurveEditorTransformTool::ScaleCenterDiameter, 0.f),
+			FVector2D(0.f, -CurveEditorTransformTool::ScaleCenterDiameter) + CircleOffset, FVector2D(0.f, -CurveEditorTransformTool::ScaleCenterDiameter), HighlightThickness, ESlateDrawEffect::None, ScaleCenterHighlightColor);
+		FSlateDrawElement::MakeSpline(OutDrawElements, InPaintOnLayerId, ScaleCenterGeometry.ToPaintGeometry(), FVector2D(CurveEditorTransformTool::ScaleCenterDiameter, 0.f) + CircleOffset, FVector2D(-CurveEditorTransformTool::ScaleCenterDiameter, 0.f),
+			FVector2D(0.f, CurveEditorTransformTool::ScaleCenterDiameter) + CircleOffset, FVector2D(0.f, CurveEditorTransformTool::ScaleCenterDiameter), HighlightThickness, ESlateDrawEffect::None, ScaleCenterHighlightColor);
+		FSlateDrawElement::MakeSpline(OutDrawElements, InPaintOnLayerId, ScaleCenterGeometry.ToPaintGeometry(), FVector2D(0.f, CurveEditorTransformTool::ScaleCenterDiameter) + CircleOffset, FVector2D(0.f, -CurveEditorTransformTool::ScaleCenterDiameter),
+			FVector2D(-CurveEditorTransformTool::ScaleCenterDiameter, 0.f) + CircleOffset, FVector2D(-CurveEditorTransformTool::ScaleCenterDiameter, 0.f), HighlightThickness, ESlateDrawEffect::None, ScaleCenterHighlightColor);
+		FSlateDrawElement::MakeSpline(OutDrawElements, InPaintOnLayerId, ScaleCenterGeometry.ToPaintGeometry(), FVector2D(0.f, -CurveEditorTransformTool::ScaleCenterDiameter) + CircleOffset, FVector2D(0.f, CurveEditorTransformTool::ScaleCenterDiameter),
+			FVector2D(CurveEditorTransformTool::ScaleCenterDiameter, 0.f) + CircleOffset, FVector2D(CurveEditorTransformTool::ScaleCenterDiameter, 0.f), HighlightThickness, ESlateDrawEffect::None, ScaleCenterHighlightColor);
+
 	}
 
 	// Draw the four corners + highlights
 	{
 		FGeometry TopLeftCornerGeometry, TopRightCornerGeometry, BottomLeftCornerGeometry, BottomRightCornerGeometry;
-		InTransformWidget.GetCornerGeometry(AllottedGeometry, TopLeftCornerGeometry, TopRightCornerGeometry, BottomLeftCornerGeometry, BottomRightCornerGeometry);
+		InTransformWidget.GetCornerGeometry(InAllottedGeometry, TopLeftCornerGeometry, TopRightCornerGeometry, BottomLeftCornerGeometry, BottomRightCornerGeometry);
 
 		const bool bTopLeft =		InTransformWidget.SelectedAnchorFlags == (ECurveEditorAnchorFlags::Top		| ECurveEditorAnchorFlags::Left );
 		const bool bTopRight =		InTransformWidget.SelectedAnchorFlags == (ECurveEditorAnchorFlags::Top		| ECurveEditorAnchorFlags::Right);
@@ -395,17 +598,17 @@ void FCurveEditorTransformTool::DrawMarqueeWidget(const FCurveEditorTransformWid
 		FLinearColor BottomRightHighlightColor	= bBottomRight  ? FLinearColor::White.CopyWithNewOpacity(CurveEditorTransformTool::EdgeHighlightAlpha) : FLinearColor::Transparent;
 
 		// Top Left (Highlight, Corner Icon)
-		FSlateDrawElement::MakeBox(OutDrawElements, PaintOnLayerId, TopLeftCornerGeometry.ToPaintGeometry(), FEditorStyle::GetBrush(TEXT("WhiteBrush")), ESlateDrawEffect::None, TopLeftHighlightColor);
-		FSlateDrawElement::MakeBox(OutDrawElements, PaintOnLayerId, TopLeftCornerGeometry.ToPaintGeometry(), FEditorStyle::GetBrush(TEXT("MarqueeSelection")));
+		FSlateDrawElement::MakeBox(OutDrawElements, InPaintOnLayerId, TopLeftCornerGeometry.ToPaintGeometry(), FEditorStyle::GetBrush(TEXT("WhiteBrush")), ESlateDrawEffect::None, TopLeftHighlightColor);
+		FSlateDrawElement::MakeBox(OutDrawElements, InPaintOnLayerId, TopLeftCornerGeometry.ToPaintGeometry(), FEditorStyle::GetBrush(TEXT("MarqueeSelection")));
 		// Top Right										 
-		FSlateDrawElement::MakeBox(OutDrawElements, PaintOnLayerId, TopRightCornerGeometry.ToPaintGeometry(), FEditorStyle::GetBrush(TEXT("WhiteBrush")), ESlateDrawEffect::None, TopRightHighlightColor);
-		FSlateDrawElement::MakeBox(OutDrawElements, PaintOnLayerId, TopRightCornerGeometry.ToPaintGeometry(), FEditorStyle::GetBrush(TEXT("MarqueeSelection")));
+		FSlateDrawElement::MakeBox(OutDrawElements, InPaintOnLayerId, TopRightCornerGeometry.ToPaintGeometry(), FEditorStyle::GetBrush(TEXT("WhiteBrush")), ESlateDrawEffect::None, TopRightHighlightColor);
+		FSlateDrawElement::MakeBox(OutDrawElements, InPaintOnLayerId, TopRightCornerGeometry.ToPaintGeometry(), FEditorStyle::GetBrush(TEXT("MarqueeSelection")));
 		// Bottom Left										 
-		FSlateDrawElement::MakeBox(OutDrawElements, PaintOnLayerId, BottomLeftCornerGeometry.ToPaintGeometry(), FEditorStyle::GetBrush(TEXT("WhiteBrush")), ESlateDrawEffect::None, BottomLeftHighlightColor);
-		FSlateDrawElement::MakeBox(OutDrawElements, PaintOnLayerId, BottomLeftCornerGeometry.ToPaintGeometry(), FEditorStyle::GetBrush(TEXT("MarqueeSelection")));
+		FSlateDrawElement::MakeBox(OutDrawElements, InPaintOnLayerId, BottomLeftCornerGeometry.ToPaintGeometry(), FEditorStyle::GetBrush(TEXT("WhiteBrush")), ESlateDrawEffect::None, BottomLeftHighlightColor);
+		FSlateDrawElement::MakeBox(OutDrawElements, InPaintOnLayerId, BottomLeftCornerGeometry.ToPaintGeometry(), FEditorStyle::GetBrush(TEXT("MarqueeSelection")));
 		// Bottom Right										 
-		FSlateDrawElement::MakeBox(OutDrawElements, PaintOnLayerId, BottomRightCornerGeometry.ToPaintGeometry(), FEditorStyle::GetBrush(TEXT("WhiteBrush")), ESlateDrawEffect::None, BottomRightHighlightColor);
-		FSlateDrawElement::MakeBox(OutDrawElements, PaintOnLayerId, BottomRightCornerGeometry.ToPaintGeometry(), FEditorStyle::GetBrush(TEXT("MarqueeSelection")));
+		FSlateDrawElement::MakeBox(OutDrawElements, InPaintOnLayerId, BottomRightCornerGeometry.ToPaintGeometry(), FEditorStyle::GetBrush(TEXT("WhiteBrush")), ESlateDrawEffect::None, BottomRightHighlightColor);
+		FSlateDrawElement::MakeBox(OutDrawElements, InPaintOnLayerId, BottomRightCornerGeometry.ToPaintGeometry(), FEditorStyle::GetBrush(TEXT("MarqueeSelection")));
 	}
 }
 
@@ -422,6 +625,8 @@ void FCurveEditorTransformTool::OnDragStart()
 
 	// We need to cache our key data because all of our calculations have to be relative to the starting data and not the current per-frame data.
 	KeysByCurve.Reset();
+	bool bMinMaxIsSet = false; // will be better cached than TOptional and we only need to check this flag here.
+
 	for (const TTuple<FCurveModelID, FKeyHandleSet>& Pair : CurveEditor->GetSelection().GetAll())
 	{
 		FCurveModelID CurveID = Pair.Key;
@@ -438,14 +643,36 @@ void FCurveEditorTransformTool::OnDragStart()
 
 			KeyData.StartKeyPositions.SetNumZeroed(KeyData.Handles.Num());
 			Curve->GetKeyPositions(KeyData.Handles, KeyData.StartKeyPositions);
+			for (const FKeyPosition& KeyPosition : KeyData.StartKeyPositions)
+			{
+				if (bMinMaxIsSet)
+				{
+					if (KeyPosition.InputValue < InputMinMax.GetLowerBoundValue())
+					{
+						InputMinMax.SetLowerBoundValue(KeyPosition.InputValue);
 		}
+					else if (KeyPosition.InputValue > InputMinMax.GetUpperBoundValue())
+					{
+						InputMinMax.SetUpperBoundValue(KeyPosition.InputValue);
+					}
 	}
+				else
+				{
+					bMinMaxIsSet = true;
+					InputMinMax = TRange<double>(KeyPosition.InputValue);
+				}
+			}
+		}
+	}
+
 	TransformWidget.StartSize = TransformWidget.Size;
 	TransformWidget.StartPosition = TransformWidget.Position;
 	SnappingState.Reset();
+	PrevState.Positon = TransformWidget.Position;
+	PrevState.Size = TransformWidget.Size;
 }
 
-void FCurveEditorTransformTool::OnDrag(const FPointerEvent& MouseEvent)
+void FCurveEditorTransformTool::OnDrag(const FPointerEvent& InMouseEvent, const FVector2D& InLocalMousePosition)
 {
 	TSharedPtr<FCurveEditor> CurveEditor = WeakCurveEditor.Pin();
 	if (!CurveEditor)
@@ -456,173 +683,154 @@ void FCurveEditorTransformTool::OnDrag(const FPointerEvent& MouseEvent)
 	TArray<FKeyPosition> NewKeyPositionScratch;
 
 	FSlateLayoutTransform ContainerToAbsolute = CurveEditor->GetPanel()->GetViewContainerGeometry().GetAccumulatedLayoutTransform().Inverse();
-
-	// Dragging the center is the easy case!
-	if (TransformWidget.SelectedAnchorFlags == ECurveEditorAnchorFlags::Center)
+	const bool bFalloffOn = InMouseEvent.IsControlDown();
+	const bool bScaleCenterOn = InMouseEvent.IsAltDown();
+	if (bFalloffOn && (TransformWidget.SelectedAnchorFlags & ECurveEditorAnchorFlags::FalloffAll) != ECurveEditorAnchorFlags::None)
 	{
-		FVector2D AxisLockedMousePosition = CurveEditor->GetAxisSnap().GetSnappedPosition(InitialMousePosition, MouseEvent.GetScreenSpacePosition(), MouseEvent, SnappingState);
+		const bool bTopLeft = ((TransformWidget.SelectedAnchorFlags & ECurveEditorAnchorFlags::FalloffTopLeft) != ECurveEditorAnchorFlags::None);
+		const bool bTopRight = ((TransformWidget.SelectedAnchorFlags & ECurveEditorAnchorFlags::FalloffTopRight) != ECurveEditorAnchorFlags::None);
+		const bool bLeft = ((TransformWidget.SelectedAnchorFlags & ECurveEditorAnchorFlags::FalloffLeft) != ECurveEditorAnchorFlags::None);
+		const bool bRight = ((TransformWidget.SelectedAnchorFlags & ECurveEditorAnchorFlags::FalloffRight) != ECurveEditorAnchorFlags::None);
+		if ((bLeft || bRight) && TransformWidget.Size.Y > SMALL_NUMBER)
 		{
-			FVector2D MouseDelta = AxisLockedMousePosition - InitialMousePosition;
-			TransformWidget.Position = TransformWidget.StartPosition + MouseDelta;
+			FVector2D MouseDelta = InLocalMousePosition - TransformWidget.Position;
+			FalloffHeight = 1.0f - MouseDelta.Y / TransformWidget.Size.Y;
+			FalloffHeight = FMath::Clamp(FalloffHeight, 0.0f, 1.0f);
 		}
-
-		for (const FKeyData& KeyData : KeysByCurve)
+		if (bTopLeft || bTopRight && TransformWidget.Size.X > SMALL_NUMBER)
 		{
-			const SCurveEditorView* View = CurveEditor->FindFirstInteractiveView(KeyData.CurveID);
-			if (!View)
+			FVector2D MouseDelta = InLocalMousePosition - TransformWidget.Position;
+
+			FalloffWidth = (MouseDelta.X) / TransformWidget.Size.X;
+			FalloffWidth = FMath::Clamp(FalloffWidth, 0.0f, 1.0f);
+			if (FalloffWidth < 0.5f)
 			{
-				continue;
+				FalloffWidth = (0.5 - FalloffWidth) * 2.0f;
 			}
-
-			FCurveModel* CurveModel = CurveEditor->FindCurve(KeyData.CurveID);
-			check(CurveModel);
-
-			FCurveEditorScreenSpace CurveSpace = View->GetCurveSpace(KeyData.CurveID);
-
-			double DeltaInput = (AxisLockedMousePosition.X - InitialMousePosition.X) / CurveSpace.PixelsPerInput();
-			double DeltaOutput = -(AxisLockedMousePosition.Y - InitialMousePosition.Y) / CurveSpace.PixelsPerOutput();
-
-			NewKeyPositionScratch.Reset();
-			NewKeyPositionScratch.Reserve(KeyData.StartKeyPositions.Num());
-
-			for (FKeyPosition StartPosition : KeyData.StartKeyPositions)
+			else if (FalloffWidth > 0.5f)
 			{
-				StartPosition.InputValue += DeltaInput;
-				StartPosition.OutputValue += DeltaOutput;
-
-				StartPosition.InputValue = View->IsTimeSnapEnabled() ? CurveEditor->GetCurveSnapMetrics(KeyData.CurveID).SnapInputSeconds(StartPosition.InputValue) : StartPosition.InputValue;
-				StartPosition.OutputValue = View->IsValueSnapEnabled() ? CurveEditor->GetCurveSnapMetrics(KeyData.CurveID).SnapOutput(StartPosition.OutputValue) : StartPosition.OutputValue;
-
-				NewKeyPositionScratch.Add(StartPosition);
+				FalloffWidth = (FalloffWidth - 0.5f) * 2.0f;
 			}
-
-			CurveModel->SetKeyPositions(KeyData.Handles, NewKeyPositionScratch);
 		}
 	}
-	else if (TransformWidget.SelectedAnchorFlags != ECurveEditorAnchorFlags::None)
+	else if (bScaleCenterOn && (TransformWidget.SelectedAnchorFlags & ECurveEditorAnchorFlags::CenterScale) != ECurveEditorAnchorFlags::None)
 	{
-		// Rescaling is where things get tricky. If a user is selecting an edge we scale on one axis at once, and if they select a corner we scale on two axis at once.
-		// If they press alt, we scale relative to the center (instead of relative to the opposite corner). Because these are all bit-flag'd together we can implement
-		// this generically by scaling coordinates relative to an arbitrary center and then masking off whether it affects the X or Y axis depending on anchor flags.
-		const bool bAffectsX = ((TransformWidget.SelectedAnchorFlags & ECurveEditorAnchorFlags::Left) != ECurveEditorAnchorFlags::None ||
-			(TransformWidget.SelectedAnchorFlags & ECurveEditorAnchorFlags::Right) != ECurveEditorAnchorFlags::None);
-
-		const bool bAffectsY = ((TransformWidget.SelectedAnchorFlags & ECurveEditorAnchorFlags::Top) != ECurveEditorAnchorFlags::None ||
-			(TransformWidget.SelectedAnchorFlags & ECurveEditorAnchorFlags::Bottom) != ECurveEditorAnchorFlags::None);
-
-		// We calculate this in [0-1] space for the widget to make the logic easier to follow.
-		FVector2D ScaleCenter = FVector2D(0.5f, 0.5f);
-		const bool bScaleFromEdge = !MouseEvent.IsAltDown();
-
-		if (bScaleFromEdge)
+		const FVector2D MouseDelta = InLocalMousePosition - TransformWidget.BoundsPosition;
+		RelativeScaleCenter = MouseDelta / TransformWidget.BoundsSize;
+	}
+	else
+	{
+		// Dragging the center is the easy case!
+		if (TransformWidget.SelectedAnchorFlags == ECurveEditorAnchorFlags::Center)
 		{
-			if ((TransformWidget.SelectedAnchorFlags & ECurveEditorAnchorFlags::Left) != ECurveEditorAnchorFlags::None)
+			const FVector2D AxisLockedMousePosition = CurveEditor->GetAxisSnap().GetSnappedPosition(InitialMousePosition, InMouseEvent.GetScreenSpacePosition(), InMouseEvent, SnappingState);
 			{
-				// Anchor to the right side
-				ScaleCenter.X = 1.0f; 
+				FVector2D MouseDelta = AxisLockedMousePosition - InitialMousePosition;
+				TransformWidget.Position = TransformWidget.StartPosition + MouseDelta;
 			}
 
-			if ((TransformWidget.SelectedAnchorFlags & ECurveEditorAnchorFlags::Right) != ECurveEditorAnchorFlags::None)
+			for (const FKeyData& KeyData : KeysByCurve)
 			{
-				// Anchor to the left side
-				ScaleCenter.X = 0.0f;
-			}
+				const SCurveEditorView* View = CurveEditor->FindFirstInteractiveView(KeyData.CurveID);
+				if (!View)
+				{
+					continue;
+				}
 
-			if ((TransformWidget.SelectedAnchorFlags & ECurveEditorAnchorFlags::Top) != ECurveEditorAnchorFlags::None)
-			{
-				// Anchor to the bottom side. Slate uses top left as 0 so we flip here.
-				ScaleCenter.Y = 1.0f;
-			}
+				FCurveModel* CurveModel = CurveEditor->FindCurve(KeyData.CurveID);
+				check(CurveModel);
 
-			if ((TransformWidget.SelectedAnchorFlags & ECurveEditorAnchorFlags::Bottom) != ECurveEditorAnchorFlags::None)
-			{
-				// Anchor to the top side
-				ScaleCenter.Y = 0.0f;
-			}
-		}
+				FCurveEditorScreenSpace CurveSpace = View->GetCurveSpace(KeyData.CurveID);
 
-		// This is the absolute change since our KeysByCurve was initialized.
-		const FVector2D MouseDelta = MouseEvent.GetScreenSpacePosition() - InitialMousePosition;
-		// TransformWidget.Size = TransformWidget.StartSize += MouseDelta;
+				const double DeltaInput = (AxisLockedMousePosition.X - InitialMousePosition.X) / CurveSpace.PixelsPerInput();
+				const double DeltaOutput = -(AxisLockedMousePosition.Y - InitialMousePosition.Y) / CurveSpace.PixelsPerOutput();
 
-		// We have to flip the delta depending on which edge you grabbed so that the change always goes towards the mouse.
-		const float InputMulSign = (TransformWidget.SelectedAnchorFlags & ECurveEditorAnchorFlags::Left) != ECurveEditorAnchorFlags::None ? -1.0f : 1.0f;
-		const float OutputMulSign = (TransformWidget.SelectedAnchorFlags & ECurveEditorAnchorFlags::Top) != ECurveEditorAnchorFlags::None ? -1.0f : 1.0f;
-		const FVector2D AxisFixedMouseDelta = FVector2D(MouseDelta.X * InputMulSign, MouseDelta.Y * OutputMulSign);
-		const FVector2D PanelSpaceCenter = TransformWidget.StartPosition + (TransformWidget.StartSize * ScaleCenter);
-		const FVector2D PercentChanged = FVector2D(1.0f, 1.0f) + (AxisFixedMouseDelta / TransformWidget.StartSize); // ie: 5 pixel change on a 100 wide gives you 1.05
+				NewKeyPositionScratch.Reset();
+				NewKeyPositionScratch.Reserve(KeyData.StartKeyPositions.Num());
 
-		// We now know if we need to affect both X and Y, and we know where we're scaling from. Now we can loop through the keys and actually modify their positions.
-		// We perform the scale on both axis (for simplicity) and then read which axis it should effect before assigning it back to the key position.
-		for (const FKeyData& KeyData : KeysByCurve)
-		{
-			const SCurveEditorView* View = CurveEditor->FindFirstInteractiveView(KeyData.CurveID);
-			if (!View)
-			{
-				continue;
-			}
-
-			FCurveModel* CurveModel = CurveEditor->FindCurve(KeyData.CurveID);
-			check(CurveModel);
+				for (FKeyPosition StartPosition : KeyData.StartKeyPositions)
+				{
+					StartPosition.InputValue += DeltaInput;
+					StartPosition.OutputValue += DeltaOutput;
 
-			// Compute the curve-space center for transformation by transforming the center from panel space to view space, then to curve space
-			FSlateLayoutTransform OuterToInnerTransform = Concatenate(ContainerToAbsolute, View->GetCachedGeometry().GetAccumulatedLayoutTransform()).Inverse();
-			FVector2D ViewSpaceCenter = OuterToInnerTransform.TransformPoint(PanelSpaceCenter);
+					StartPosition.InputValue = View->IsTimeSnapEnabled() ? CurveEditor->GetCurveSnapMetrics(KeyData.CurveID).SnapInputSeconds(StartPosition.InputValue) : StartPosition.InputValue;
+					StartPosition.OutputValue = View->IsValueSnapEnabled() ? CurveEditor->GetCurveSnapMetrics(KeyData.CurveID).SnapOutput(StartPosition.OutputValue) : StartPosition.OutputValue;
 
-			FCurveEditorScreenSpace CurveSpace = View->GetCurveSpace(KeyData.CurveID);
-			const double CurveSpaceCenterInput =  CurveSpace.ScreenToSeconds(ViewSpaceCenter.X);
-			const double CurveSpaceCenterOutput = CurveSpace.ScreenToValue(ViewSpaceCenter.Y);
+					NewKeyPositionScratch.Add(StartPosition);
+				}
 
-			// UE_LOG(LogTemp, Log, TEXT("ScaleCenter Input: %f ScaleCenter Output: %f PercentChanged: %s"), CenterInput, CenterOutput, *PercentChanged.ToString());
+				CurveModel->SetKeyPositions(KeyData.Handles, NewKeyPositionScratch);
+			}
+		}
+		else if (TransformWidget.SelectedAnchorFlags != ECurveEditorAnchorFlags::None)
+		{
+			const bool bAffectsX = ((TransformWidget.SelectedAnchorFlags & ECurveEditorAnchorFlags::Left) != ECurveEditorAnchorFlags::None ||
+				(TransformWidget.SelectedAnchorFlags & ECurveEditorAnchorFlags::Right) != ECurveEditorAnchorFlags::None);
 
-			NewKeyPositionScratch.Reset();
-			NewKeyPositionScratch.Reserve(KeyData.StartKeyPositions.Num());
+			const bool bAffectsY = ((TransformWidget.SelectedAnchorFlags & ECurveEditorAnchorFlags::Top) != ECurveEditorAnchorFlags::None ||
+				(TransformWidget.SelectedAnchorFlags & ECurveEditorAnchorFlags::Bottom) != ECurveEditorAnchorFlags::None);
 
-			// UE_LOG(LogTemp, Log, TEXT("MouseDelta: %s PercentChanged: %s CenterInput: %f CenterOutput: %f"),
-			// *MouseDelta.ToString(), *PercentChanged.ToString(), CenterInput, CenterOutput);
+			// We calculate this in [0-1] space for the widget to make the logic easier to follow.
+			FVector2D ScaleCenter = RelativeScaleCenter;
+			const bool bScaleFromEdge = !InMouseEvent.IsAltDown();
 
-			int32 i = 0;
-			for (FKeyPosition StartPosition : KeyData.StartKeyPositions)
+			if (bScaleFromEdge)
 			{
-				// Step 1 is to rescale the position of the key by the percentage change on each axis.
-				double ScaledInput = (StartPosition.InputValue - CurveSpaceCenterInput) * PercentChanged.X; // *InputMulSign;
-				double ScaledOutput = (StartPosition.OutputValue - CurveSpaceCenterOutput) * PercentChanged.Y; // *OutputMulSign;
-
-				// Step 2 is to subtract it from the center position so we support scaling from places other than zero
-				double NewInput = CurveSpaceCenterInput + ScaledInput;
-				double NewOutput = CurveSpaceCenterOutput + ScaledOutput;
-
-				// UE_LOG(LogTemp, Log, TEXT("[%d] InputValue: %f OutputValue: %f ScaledInput: %f ScaledOutput: %f NewInput: %f NewOutput: %f"),
-				// i, StartPosition.InputValue, StartPosition.OutputValue, ScaledInput, ScaledOutput, NewInput, NewOutput);
-				i++;
-
-				
-				// Snap the new values to the grid. We calculate both X and Y changes for ease of programming above and just limit which one it applies to.
-				// This includes snapping, otherwise dragging on an edge can cause it to snap on the opposite axis.
-				if (bAffectsX)
+				if ((TransformWidget.SelectedAnchorFlags & ECurveEditorAnchorFlags::Left) != ECurveEditorAnchorFlags::None)
 				{
-					StartPosition.InputValue = View->IsTimeSnapEnabled() ? CurveEditor->GetCurveSnapMetrics(KeyData.CurveID).SnapInputSeconds(NewInput) : NewInput;
+					// Anchor to the right side
+					ScaleCenter.X = 1.0f; 
 				}
-				if (bAffectsY)
+
+				if ((TransformWidget.SelectedAnchorFlags & ECurveEditorAnchorFlags::Right) != ECurveEditorAnchorFlags::None)
 				{
-					StartPosition.OutputValue = View->IsValueSnapEnabled() ? CurveEditor->GetCurveSnapMetrics(KeyData.CurveID).SnapOutput(NewOutput) : NewOutput;
+					// Anchor to the left side
+					ScaleCenter.X = 0.0f;
 				}
 
+				if ((TransformWidget.SelectedAnchorFlags & ECurveEditorAnchorFlags::Top) != ECurveEditorAnchorFlags::None)
+				{
+					// Anchor to the bottom side. Slate uses top left as 0 so we flip here.
+					ScaleCenter.Y = 1.0f;
+				}
 
-				NewKeyPositionScratch.Add(StartPosition);
+				if ((TransformWidget.SelectedAnchorFlags & ECurveEditorAnchorFlags::Bottom) != ECurveEditorAnchorFlags::None)
+				{
+					// Anchor to the top side
+					ScaleCenter.Y = 0.0f;
+				}
 			}
 
-			CurveModel->SetKeyPositions(KeyData.Handles, NewKeyPositionScratch);
+			// This is the absolute change since our KeysByCurve was initialized.
+			const FVector2D MouseDelta = InMouseEvent.GetScreenSpacePosition() - InitialMousePosition;
+			// TransformWidget.Size = TransformWidget.StartSize += MouseDelta;
+
+			// We have to flip the delta depending on which edge you grabbed so that the change always goes towards the mouse.
+			const float InputMulSign = (TransformWidget.SelectedAnchorFlags & ECurveEditorAnchorFlags::Left) != ECurveEditorAnchorFlags::None ? -1.0f : 1.0f;
+			const float OutputMulSign = (TransformWidget.SelectedAnchorFlags & ECurveEditorAnchorFlags::Top) != ECurveEditorAnchorFlags::None ? -1.0f : 1.0f;
+			const FVector2D AxisFixedMouseDelta = FVector2D(MouseDelta.X * InputMulSign, MouseDelta.Y * OutputMulSign);
+			const FVector2D PanelSpaceCenter = TransformWidget.StartPosition + (TransformWidget.StartSize * ScaleCenter);
+			const FVector2D ChangeAmount = (AxisFixedMouseDelta / TransformWidget.StartSize);
+
+			ScaleFrom(PanelSpaceCenter, ChangeAmount, bFalloffOn, bAffectsX, bAffectsY);
 		}
 	}
+
+	UpdateMarqueeBoundingBox();
+
+	UpdateToolOptions();
+
+	PrevState.Positon = TransformWidget.Position;
+	PrevState.Size = TransformWidget.Size;
 }
 
 void FCurveEditorTransformTool::OnDragEnd()
 {
 	// This finalizes the transaction
 	ActiveTransaction.Reset();
-	// PreDragCurveData.Reset();
-	// TotalDeltaTime = 0.0;
+
+	PrevState.Positon = TransformWidget.Position;
+	PrevState.Size = TransformWidget.Size;
 }
 
 void FCurveEditorTransformTool::StopDragIfPossible()
@@ -638,12 +846,127 @@ void FCurveEditorTransformTool::StopDragIfPossible()
 	}
 }
 
+void FCurveEditorTransformTool::ScaleFrom(const FVector2D& InPanelSpaceCenter, const FVector2D& InChangeAmount, const bool bInFalloffOn, const bool bInAffectsX, const bool bInAffectsY) const
+{
+	TSharedPtr<FCurveEditor> CurveEditor = WeakCurveEditor.Pin();
+	if (!CurveEditor)
+	{
+		return;
+	}
+
+	// Rescaling is where things get tricky. If a user is selecting an edge we scale on one axis at once, and if they select a corner we scale on two axis at once.
+	// If they press alt, we scale relative to the center (instead of relative to the opposite corner). Because these are all bit-flag'd together we can implement
+	// this generically by scaling coordinates relative to an arbitrary center
+
+	FVector2D PercentChanged = FVector2D(1.0f, 1.0f) + InChangeAmount; // ie: 5 pixel change on a 100 wide gives you 1.05
+
+	FSlateLayoutTransform ContainerToAbsolute = CurveEditor->GetPanel()->GetViewContainerGeometry().GetAccumulatedLayoutTransform().Inverse();
+
+	TArray<FKeyPosition> NewKeyPositionScratch;
+
+	// We now know if we need to affect both X and Y, and we know where we're scaling from. Now we can loop through the keys and actually modify their positions.
+	// We perform the scale on both axis (for simplicity) and then read which axis it should effect before assigning it back to the key position.
+	for (const FKeyData& KeyData : KeysByCurve)
+	{
+		const SCurveEditorView* View = CurveEditor->FindFirstInteractiveView(KeyData.CurveID);
+		if (!View)
+		{
+			continue;
+		}
+
+		FCurveModel* CurveModel = CurveEditor->FindCurve(KeyData.CurveID);
+		check(CurveModel);
+
+		// Compute the curve-space center for transformation by transforming the center from panel space to view space, then to curve space
+		FSlateLayoutTransform OuterToInnerTransform = Concatenate(ContainerToAbsolute, View->GetCachedGeometry().GetAccumulatedLayoutTransform()).Inverse();
+		FVector2D ViewSpaceCenter = OuterToInnerTransform.TransformPoint(InPanelSpaceCenter);
+
+		FCurveEditorScreenSpace CurveSpace = View->GetCurveSpace(KeyData.CurveID);
+		const double CurveSpaceCenterInput = CurveSpace.ScreenToSeconds(ViewSpaceCenter.X);
+		const double CurveSpaceCenterOutput = CurveSpace.ScreenToValue(ViewSpaceCenter.Y);
+
+		NewKeyPositionScratch.Reset();
+		NewKeyPositionScratch.Reserve(KeyData.StartKeyPositions.Num());
+
+		FCurveSnapMetrics SnapMetrics = CurveEditor->GetCurveSnapMetrics(KeyData.CurveID);
+
+		for (FKeyPosition StartPosition : KeyData.StartKeyPositions)
+		{
+			// Step 1 is to modify the Y percent changed if falloff is on
+			if (bInFalloffOn)
+			{
+				const double FalloffScale = GetFalloffWeight(StartPosition.InputValue);
+				PercentChanged.Y = 1.0f + InChangeAmount.Y * FalloffScale;
+			}
+
+			// Step 2 is to rescale the position of the key by the percentage change on each axis.
+			const double ScaledInput = (StartPosition.InputValue - CurveSpaceCenterInput) * PercentChanged.X; // *InputMulSign;
+			const double ScaledOutput = (StartPosition.OutputValue - CurveSpaceCenterOutput) * PercentChanged.Y; // *OutputMulSign;
+
+			// Step 3 is to subtract it from the center position so we support scaling from places other than zero
+			const double NewInput = CurveSpaceCenterInput + ScaledInput;
+			const double NewOutput = CurveSpaceCenterOutput + ScaledOutput;
+
+			// Snap the new values to the grid. We calculate both X and Y changes for ease of programming above and just limit which one it applies to.
+			// This includes snapping, otherwise dragging on an edge can cause it to snap on the opposite axis.
+			if (bInAffectsX)
+			{
+				StartPosition.InputValue = View->IsTimeSnapEnabled() ? SnapMetrics.SnapInputSeconds(NewInput) : NewInput;
+			}
+			if (bInAffectsY)
+			{
+				StartPosition.OutputValue = View->IsValueSnapEnabled() ? SnapMetrics.SnapOutput(NewOutput) : NewOutput;
+			}
+
+			NewKeyPositionScratch.Add(StartPosition);
+		}
+
+		CurveModel->SetKeyPositions(KeyData.Handles, NewKeyPositionScratch);
+	}
+}
+
 void FCurveEditorTransformTool::Tick(const FGeometry& AllottedGeometry, const double InCurrentTime, const float InDeltaTime)
 {
 	// We update the size and position of the box every frame as some scale operations aren't 1:1
 	// so this keeps the box visually containing all keys even if the mouse position no longer quite 
 	// matches up.
 	UpdateMarqueeBoundingBox();
+
+	TSharedPtr<FCurveEditor> CurveEditor = WeakCurveEditor.Pin();
+	if (!CurveEditor)
+	{
+		return;
+	}
+
+	if (!FSlateApplication::Get().GetModifierKeys().IsAltDown() && !RelativeScaleCenter.Equals(FVector2D(.5f, .5f)))
+	{
+		RelativeScaleCenter = FVector2D(.5f, .5f);
+	}
+
+	// remove tool options widget if less than two selected or there are two different curves
+	TArray<FCurveModelID> CurveModels1;
+	TArray<FCurveModelID> CurveModels2; // no way to reset the curve models between if statements so there has to be two
+	if ((CurveEditor->GetSelection().Count() <= 1 || CurveEditor->GetSelection().GetAll().GetKeys(CurveModels1) != 1) 
+		&& ToolOptionsOnScope != nullptr)
+	{
+		ToolOptionsOnScope = nullptr;
+		OnOptionsRefreshDelegate.Broadcast();
+	}
+	else if ((CurveEditor->GetSelection().Count() > 1 && CurveEditor->GetSelection().GetAll().GetKeys(CurveModels2) == 1) 
+		&& ToolOptionsOnScope == nullptr)
+	{
+		ToolOptionsOnScope = MakeShared<FStructOnScope>(FTransformToolOptions::StaticStruct(), (uint8*)&ToolOptions);
+		OnOptionsRefreshDelegate.Broadcast();
+	}
+
+	// check if there have been any changes made outside of our tool
+	if (PrevState.Positon != TransformWidget.BoundsPosition || PrevState.Size != TransformWidget.BoundsSize)
+	{
+		UpdateToolOptions();
+
+		PrevState.Positon = TransformWidget.Position;
+		PrevState.Size = TransformWidget.Size;
+	}
 }
 
 void FCurveEditorTransformTool::BindCommands(TSharedRef<FUICommandList> CommandBindings)
@@ -657,4 +980,72 @@ void FCurveEditorTransformTool::BindCommands(TSharedRef<FUICommandList> CommandB
 		CommandBindings->MapAction(FCurveEditorToolCommands::Get().ActivateTransformTool, ActivateTransformTool, FCanExecuteAction(), TransformToolIsActive);
 	}
 }
+
+double FCurveEditorTransformTool::GetFalloffWeight(double InputValue) const
+{
+	const double InputSize = (InputMinMax.GetUpperBoundValue() - InputMinMax.GetLowerBoundValue());
+	const double InputMidPoint = InputSize * 0.5;
+	const double MidPointDiff = InputMidPoint * FalloffWidth;
+	TRange<double> InputMiddleRange(InputMidPoint - MidPointDiff, InputMidPoint + MidPointDiff);
+	if (InputValue < InputMiddleRange.GetLowerBoundValue())
+	{
+		const double LinearScale = (InputValue - InputMinMax.GetLowerBoundValue()) / (InputMiddleRange.GetLowerBoundValue() - InputMinMax.GetLowerBoundValue());
+		double  Weight = (LinearScale * (1.0f - FalloffHeight) + FalloffHeight);
+		Weight = ModifyWeightByInterpType(Weight);
+		return Weight;
+	}
+	else if (InputValue > InputMiddleRange.GetUpperBoundValue())
+	{
+		const double  LinearScale = 1.0 - (InputValue - InputMiddleRange.GetUpperBoundValue()) / (InputMinMax.GetUpperBoundValue() - InputMiddleRange.GetUpperBoundValue());
+		double  Weight = (LinearScale * (1.0f - FalloffHeight) + FalloffHeight);
+		Weight = ModifyWeightByInterpType(Weight);
+		return Weight;
+	}
+
+	return 1.0;
+}
+
+double FCurveEditorTransformTool::ModifyWeightByInterpType(double Value) const
+{
+	float Result = Value;
+	switch (FalloffInterpType)
+	{
+		case EToolTransformInterpType::Linear:
+		{
+			Result = FMath::Clamp<float>(Value, 0.f, 1.f);
+			break;
+		}
+		case EToolTransformInterpType::Sinusoidal:
+		{
+			Result = FMath::Clamp<float>((FMath::Sin(Value * PI - HALF_PI) + 1.f) / 2.f, 0.f, 1.f);
+			break;
+		}
+		case EToolTransformInterpType::Cubic:
+		{
+			Result = FMath::Clamp<float>(FMath::CubicInterp<float>(0.f, 0.f, 1.f, 0.f, Value), 0.f, 1.f);
+			break;
+		}
+		case EToolTransformInterpType::CircularIn:
+		{
+			Result = FMath::Clamp<float>(FMath::InterpCircularIn<float>(0.0f, 1.0f, Value), 0.0f, 1.0f);
+			break;
+		}
+		case EToolTransformInterpType::CircularOut:
+		{
+			Result = FMath::Clamp<float>(FMath::InterpCircularOut<float>(0.0f, 1.0f, Value), 0.0f, 1.0f);
+			break;
+		}
+		case EToolTransformInterpType::ExpIn:
+		{
+			Result = FMath::Clamp<float>(FMath::InterpExpoIn<float>(0.0f, 1.0f, Value), 0.0f, 1.0f);
+			break;
+		}
+		case EToolTransformInterpType::ExpOut:
+		{
+			Result = FMath::Clamp<float>(FMath::InterpExpoOut<float>(0.0f, 1.0f, Value), 0.0f, 1.0f);
+			break;
+		}
+	}
+	return Result;
+}
 #undef LOCTEXT_NAMESPACE // "CurveEditorToolCommands"
\ No newline at end of file
diff --git a/Engine/Plugins/Editor/CurveEditorTools/Source/CurveEditorTools/Private/Tools/CurveEditorTransformTool.h b/Engine/Plugins/Editor/CurveEditorTools/Source/CurveEditorTools/Private/Tools/CurveEditorTransformTool.h
--- a/Engine/Plugins/Editor/CurveEditorTools/Source/CurveEditorTools/Private/Tools/CurveEditorTransformTool.h
+++ b/Engine/Plugins/Editor/CurveEditorTools/Source/CurveEditorTools/Private/Tools/CurveEditorTransformTool.h
@@ -10,17 +10,36 @@
 #include "ScopedTransaction.h"
 #include "CurveEditorSnapMetrics.h"
 
+#include "CurveEditorTransformTool.generated.h"
+
 class FCurveEditor;
 struct FPointerEvent;
 
-enum class ECurveEditorAnchorFlags : uint8
+enum class ECurveEditorAnchorFlags : uint16
 {
 	None = 0x00,
 	Top = 0x01,
 	Left = 0x02,
 	Right = 0x04,
 	Bottom = 0x08,
-	Center = 0x10
+	Center = 0x10,
+	FalloffTopLeft = 0x20,
+	FalloffTopRight = 0x40,
+	FalloffLeft = 0x80,
+	FalloffRight = 0x100,
+	FalloffAll = FalloffTopLeft | FalloffTopRight | FalloffLeft | FalloffRight,
+	CenterScale = 0x200
+};
+
+enum class EToolTransformInterpType : uint8
+{
+	Linear,
+	Sinusoidal,
+	Cubic,
+	CircularIn,
+	CircularOut,
+	ExpIn,
+	ExpOut,
 };
 
 ENUM_CLASS_FLAGS(ECurveEditorAnchorFlags);
@@ -35,34 +54,61 @@ struct FCurveEditorTransformWidget
 
 public:
 	ECurveEditorAnchorFlags SelectedAnchorFlags;
-	
+
 	FGeometry MakeGeometry(const FGeometry& InWidgetGeometry) const
 	{
 		return InWidgetGeometry.MakeChild(Size, FSlateLayoutTransform(Position));
 	}
 
-	ECurveEditorAnchorFlags GetAnchorFlagsForMousePosition(const FGeometry& InWidgetGeometry, const FVector2D& InMouseScreenPosition) const;
+	ECurveEditorAnchorFlags GetAnchorFlagsForMousePosition(const FGeometry& InWidgetGeometry, float  FalloffHeight, float FalloffWidth, const FVector2D& RelativeCenterScale, const FVector2D& InMouseScreenPosition) const;
 
 	void GetCenterGeometry(const FGeometry& InWidgetGeometry, FGeometry& OutCenter) const;
 	void GetSidebarGeometry(const FGeometry& InWidgetGeometry, FGeometry& OutLeft, FGeometry& OutRight, FGeometry& OutTop, FGeometry& OutBottom) const;
 	void GetCornerGeometry(const FGeometry& InWidgetGeometry, FGeometry& OutTopLeft, FGeometry& OutTopRight, FGeometry& OutBottomLeft, FGeometry& OutBottomRight) const;
+	void GetFalloffGeometry(const FGeometry& InWidgetGeometry, float FalloffHeight, float FalloffWidth, FGeometry& OutTopLeft, FGeometry& OutTopRight, FGeometry& OutLeft, FGeometry& OutRight) const;
+	void GetScaleCenterGeometry(const FGeometry& InWidgetGeometry, FVector2D ScaleCenter, FGeometry& OutInputBoxGeometry) const;
+
 
 	FVector2D Size;
 	FVector2D Position;
 	bool Visible;
 
+	// stores the physical bounds instead of the marquee bounds, for scale calculations
+	FVector2D BoundsSize;
+	FVector2D BoundsPosition;
+
 	FVector2D StartSize;
 	FVector2D StartPosition;
 };
 
+USTRUCT()
+struct FTransformToolOptions
+{
+	GENERATED_BODY()
+
+	UPROPERTY(Transient, EditAnywhere, Category = ToolOptions)
+	float UpperBound;
+
+	UPROPERTY(Transient, EditAnywhere, Category = ToolOptions)
+	float LowerBound;
+
+	UPROPERTY(Transient, EditAnywhere, Category = ToolOptions)
+	float LeftBound;
+
+	UPROPERTY(Transient, EditAnywhere, Category = ToolOptions)
+	float RightBound;
+};
 
 class FCurveEditorTransformTool : public ICurveEditorToolExtension
 {
 public:
 	FCurveEditorTransformTool(TWeakPtr<FCurveEditor> InCurveEditor)
-		: WeakCurveEditor(InCurveEditor)
-	{
-	}
+		: WeakCurveEditor(InCurveEditor), 
+		FalloffInterpType(EToolTransformInterpType::Linear),
+		FalloffHeight(0.0f), 
+		FalloffWidth(0.0f), 
+		RelativeScaleCenter(FVector2D(0.5f, 0.5f))
+	{}
 
 	// ICurveEditorToolExtension Interface
 	virtual void OnPaint(const FPaintArgs& Args, const FGeometry& AllottedGeometry, const FSlateRect& MyCullingRect, FSlateWindowElementList& OutDrawElements, int32 PaintOnLayerId, const FWidgetStyle& InWidgetStyle, bool bParentEnabled) const override;
@@ -74,16 +120,21 @@ public:
 	virtual FReply OnMouseMove(TSharedRef<SWidget> OwningWidget, const FGeometry& MyGeometry, const FPointerEvent& MouseEvent) override;
 	virtual FReply OnMouseButtonUp(TSharedRef<SWidget> OwningWidget, const FGeometry& MyGeometry, const FPointerEvent& MouseEvent) override;
 	virtual void OnFocusLost(const FFocusEvent& InFocusEvent);
+	TSharedPtr<FStructOnScope> GetToolOptions() const override { return ToolOptionsOnScope; }
+	virtual void OnToolOptionsUpdated(const FPropertyChangedEvent& InPropertyChangedEvent) override;
 	// ~ICurveEditorToolExtension
 
 private:
 	void UpdateMarqueeBoundingBox();
-	void DrawMarqueeWidget(const FCurveEditorTransformWidget& InTransformWidget, const FPaintArgs& Args, const FGeometry& AllottedGeometry, const FSlateRect& MyCullingRect, FSlateWindowElementList& OutDrawElements, int32 PaintOnLayerId, const FWidgetStyle& InWidgetStyle, bool bParentEnabled) const;
+	void UpdateToolOptions();
+	void DrawMarqueeWidget(const FCurveEditorTransformWidget& InTransformWidget, const FPaintArgs& InArgs, const FGeometry& InAllottedGeometry, const FSlateRect& InMyCullingRect, FSlateWindowElementList& OutDrawElements, const int32 InPaintOnLayerId, const FWidgetStyle& InWidgetStyle, const bool bInParentEnabled) const;
+	void ScaleFrom(const FVector2D& InPanelSpaceCenter, const FVector2D& InChangeAmount, const bool bInFalloffOn, const bool bInAffectsX, const bool bInAffectsY) const;
 
 	void OnDragStart();
-	void OnDrag(const FPointerEvent& MouseEvent);
+	void OnDrag(const FPointerEvent& InMouseEvent, const FVector2D& InLocalMousePosition);
 	void OnDragEnd();
 	void StopDragIfPossible();
+
 private:
 	/** Weak pointer back to the Curve Editor this belongs to. */
 	TWeakPtr<FCurveEditor> WeakCurveEditor;
@@ -114,8 +165,34 @@ private:
 
 	/** Key dragging data stored per-curve */
 	TArray<FKeyData> KeysByCurve;
+	TRange<double> InputMinMax;
 
 	FVector2D InitialMousePosition;
-
 	FCurveEditorAxisSnap::FSnapState SnappingState;
-};
\ No newline at end of file
+
+	// UStruct that displays properties to be modified on screen
+	FTransformToolOptions ToolOptions;
+	TSharedPtr<FStructOnScope> ToolOptionsOnScope;
+
+	// helper struct for checking if tool has changed outside of own operations, i.e. middle click
+	struct ToolBounds { FVector2D Positon, Size; };
+	ToolBounds PrevState;
+
+private:
+	/** specifies the falloff curve applied to soft selection */
+	EToolTransformInterpType FalloffInterpType;
+	/** Soft Selection Height, if 0.0 as we go up to 1.0 we go up to the top edge */
+	float FalloffHeight;
+	/** Soft Selection Width, if 0.0 we peak directly at the midpoint gradually scaling out till we reach 1.0fff */
+	float FalloffWidth;
+
+	/** point to scale from, relative to current bounds */
+	FVector2D RelativeScaleCenter;
+
+	/** Get Soft Select Change Based upon Input Value and internal soft selection weight values*/
+	double GetFalloffWeight(double InputValue) const;
+
+	/** Modify the Falloff Weight by type of Interpolation */
+	double ModifyWeightByInterpType(double Weight) const;
+};
+
diff --git a/Engine/Source/Editor/CurveEditor/Private/SCurveEditorPanel.cpp b/Engine/Source/Editor/CurveEditor/Private/SCurveEditorPanel.cpp
--- a/Engine/Source/Editor/CurveEditor/Private/SCurveEditorPanel.cpp
+++ b/Engine/Source/Editor/CurveEditor/Private/SCurveEditorPanel.cpp
@@ -37,6 +37,7 @@
 #include "CommonFrameRates.h"
 #include "CurveEditorViewRegistry.h"
 #include "SCurveEditorViewContainer.h"
+#include "SCurveEditorToolProperties.h"
 #include "Fonts/FontMeasure.h"
 #include "CurveEditorHelpers.h"
 
@@ -139,6 +140,8 @@ void SCurveEditorPanel::Construct(const FArguments& InArgs, TSharedRef<FCurveEdi
 	CurveEditor->BindCommands();
 	CurveEditor->SetTimeSliderController(InArgs._ExternalTimeSliderController);
 
+	CurveEditor->OnActiveToolChangedDelegate.AddSP(this, &SCurveEditorPanel::OnCurveEditorToolChanged);
+
 	CommandList = MakeShared<FUICommandList>();
 	CommandList->Append(InCurveEditor->GetCommands().ToSharedRef());
 
@@ -200,6 +203,14 @@ void SCurveEditorPanel::Construct(const FArguments& InArgs, TSharedRef<FCurveEdi
 				[
 					ScrollBar
 				]
+
+				+ SOverlay::Slot()
+				.Padding(10.0f, 10.0f)
+				.HAlign(HAlign_Right)
+				.VAlign(VAlign_Bottom)
+				[
+					SAssignNew(ToolPropertiesPanel, SCurveEditorToolProperties, InCurveEditor, FCurveEditorToolID::Unset())
+				]
 			]
 		]
 
@@ -660,6 +671,11 @@ void SCurveEditorPanel::UpdateCommonCurveInfo()
 	CachedCommonKeyAttributes = AccumulatedKeyAttributes.Get(FKeyAttributes());
 }
 
+void SCurveEditorPanel::OnCurveEditorToolChanged(FCurveEditorToolID InToolId)
+{
+	ToolPropertiesPanel->OnToolChanged(InToolId);
+}
+
 void SCurveEditorPanel::UpdateEditBox()
 {
 	const FCurveEditorSelection& Selection = CurveEditor->Selection;
@@ -1071,8 +1087,8 @@ TSharedPtr<FExtender> SCurveEditorPanel::GetToolbarExtender()
 				ToolBarBuilder.AddComboButton(
 					FUIAction(),
 					FOnGetContent::CreateSP(InEditorPanel, &SCurveEditorPanel::MakeGridSpacingMenu),
-					LOCTEXT("ValueSnappingOptions", "Value Snapping"),
-					LOCTEXT("ValueSnappingOptionsToolTip", "Choose what precision the Value axis is snapped to while moving keys."),
+					LOCTEXT("GridSnappingOptions", "Grid Snapping"),
+					LOCTEXT("GridSnappingOptionsToolTip", "Choose the spacing between horizontal grid lines."),
 					TAttribute<FSlateIcon>(),
 					true);
 
diff --git a/Engine/Source/Editor/CurveEditor/Private/SCurveEditorToolProperties.cpp b/Engine/Source/Editor/CurveEditor/Private/SCurveEditorToolProperties.cpp
new file mode 100644
--- /dev/null
+++ b/Engine/Source/Editor/CurveEditor/Private/SCurveEditorToolProperties.cpp
@@ -0,0 +1,85 @@
+// Copyright 1998-2019 Epic Games, Inc. All Rights Reserved.
+
+#include "SCurveEditorToolProperties.h"
+#include "Modules/ModuleManager.h"
+#include "IStructureDetailsView.h"
+#include "PropertyEditorDelegates.h"
+#include "PropertyEditorModule.h"
+
+
+void SCurveEditorToolProperties::Construct(const FArguments& InArgs, TSharedRef<FCurveEditor> InCurveEditor, FCurveEditorToolID InToolId)
+{
+	WeakCurveEditor = InCurveEditor;
+	ToolId = InToolId;
+
+	FPropertyEditorModule& PropertyEditorModule = FModuleManager::GetModuleChecked<FPropertyEditorModule>("PropertyEditor");
+	FDetailsViewArgs DetailsViewArgs;
+	DetailsViewArgs.bAllowSearch = false;
+	DetailsViewArgs.bShowScrollBar = false;
+	DetailsViewArgs.NameAreaSettings = FDetailsViewArgs::ObjectsUseNameArea;
+
+	FStructureDetailsViewArgs StructureDetailsViewArgs;
+	StructureDetailsViewArgs.bShowObjects = true;
+	StructureDetailsViewArgs.bShowAssets = true;
+	StructureDetailsViewArgs.bShowClasses = true;
+	StructureDetailsViewArgs.bShowInterfaces = true;
+
+	DetailsView = PropertyEditorModule.CreateStructureDetailView(DetailsViewArgs, StructureDetailsViewArgs, nullptr);
+
+	DetailsView->GetOnFinishedChangingPropertiesDelegate().AddSP(this, &SCurveEditorToolProperties::OnFinishedChangingProperties);
+
+	OnToolChanged(ToolId);
+
+	ChildSlot
+	[
+		DetailsView->GetWidget().ToSharedRef()
+	];
+}
+
+void SCurveEditorToolProperties::OnToolChanged(FCurveEditorToolID NewToolId)
+{
+	ToolId = NewToolId;
+
+	TSharedPtr<FCurveEditor> CurveEditor = WeakCurveEditor.Pin();
+
+	if (CurveEditor)
+	{
+		TSharedPtr<FOnOptionsRefresh> OnOptionsRefreshDelegate;
+		if (CurveEditor->GetToolExtensions().Contains(ToolId))
+		{
+			CurveEditor->GetToolExtensions()[ToolId]->OnOptionsRefreshDelegate.AddSP(this, &SCurveEditorToolProperties::RebuildProperties);
+		}
+	}
+	RebuildProperties();
+}
+
+void SCurveEditorToolProperties::RebuildProperties()
+{
+	TSharedPtr<FCurveEditor> CurveEditor = WeakCurveEditor.Pin();
+	if (CurveEditor)
+	{
+		TSharedPtr<FStructOnScope> ToolOptions;
+		if (CurveEditor->GetToolExtensions().Contains(ToolId) &&
+			(ToolOptions = CurveEditor->GetToolExtensions()[ToolId]->GetToolOptions()))
+		{
+			DetailsView->GetWidget()->SetVisibility(EVisibility::Visible);
+			DetailsView->SetStructureData(ToolOptions);
+		}
+		else
+		{
+			DetailsView->GetWidget()->SetVisibility(EVisibility::Hidden);
+		}
+	}
+}
+
+void SCurveEditorToolProperties::OnFinishedChangingProperties(const FPropertyChangedEvent& PropertyChangedEvent)
+{
+	TSharedPtr<FCurveEditor> CurveEditor = WeakCurveEditor.Pin();
+	if (CurveEditor)
+	{
+		if (CurveEditor->GetCurrentTool())
+		{
+			CurveEditor->GetCurrentTool()->OnToolOptionsUpdated(PropertyChangedEvent);
+		}
+	}
+}
\ No newline at end of file
diff --git a/Engine/Source/Editor/CurveEditor/Private/SCurveEditorToolProperties.h b/Engine/Source/Editor/CurveEditor/Private/SCurveEditorToolProperties.h
new file mode 100644
--- /dev/null
+++ b/Engine/Source/Editor/CurveEditor/Private/SCurveEditorToolProperties.h
@@ -0,0 +1,50 @@
+// Copyright 1998-2019 Epic Games, Inc. All Rights Reserved.
+
+#pragma once
+
+#include "Widgets/SCompoundWidget.h"
+#include "CurveEditor.h"
+#include "UObject/Object.h"
+
+class IStructureDetailsView;
+
+/**
+ * Widget to show curve editor tool extension properties
+ */
+class CURVEEDITOR_API SCurveEditorToolProperties : public SCompoundWidget
+{
+public:
+	SLATE_BEGIN_ARGS(SCurveEditorToolProperties)
+	{}
+
+	SLATE_END_ARGS()
+
+	SCurveEditorToolProperties() : ToolId(FCurveEditorToolID::Unset()) {}
+
+	void Construct(const FArguments& InArgs, TSharedRef<FCurveEditor> InCurveEditor, FCurveEditorToolID ToolId);
+
+	/**
+	 * Called when the curve editor tool changes
+	 * Sets the new tool ID and rebuilds the details view
+	 */
+	void OnToolChanged(FCurveEditorToolID NewToolId);
+
+private:
+
+	/**
+	 * Rebuilds details view
+	 * Called either by the tool itself when it's properties change or when the curve editor tool changes
+	 */
+	void RebuildProperties();
+
+	/**
+	 * Called when the user enters in a new value for one of the properties
+	 * Calls OnToolOptionsUpdated on the current tool
+	 */
+	void OnFinishedChangingProperties(const FPropertyChangedEvent& PropertyChangedEvent);
+
+private:
+	TWeakPtr<FCurveEditor> WeakCurveEditor;
+	TSharedPtr<IStructureDetailsView> DetailsView;
+	FCurveEditorToolID ToolId;
+};
\ No newline at end of file
diff --git a/Engine/Source/Editor/CurveEditor/Public/ICurveEditorToolExtension.h b/Engine/Source/Editor/CurveEditor/Public/ICurveEditorToolExtension.h
--- a/Engine/Source/Editor/CurveEditor/Public/ICurveEditorToolExtension.h
+++ b/Engine/Source/Editor/CurveEditor/Public/ICurveEditorToolExtension.h
@@ -4,6 +4,7 @@
 #include "Templates/SharedPointer.h"
 #include "Input/Reply.h"
 
+class FStructOnScope;
 class FUICommandList;
 class FPaintArgs;
 class FSlateRect;
@@ -74,6 +75,8 @@ private:
 	uint32 ID;
 };
 
+typedef TMulticastDelegate<void> FOnOptionsRefresh;
+
 /**
 * You can extend the Curve Editor toolset by implementing this interface. The Curve Editor guarantees that only
 * one tool will be active at any given time. A tool needs to specify if they handled certain mouse events so that
@@ -98,6 +101,8 @@ public:
 	virtual FReply OnMouseMove(TSharedRef<SWidget> OwningWidget, const FGeometry& MyGeometry, const FPointerEvent& MouseEvent) { return FReply::Unhandled(); }
 	virtual FReply OnMouseButtonDoubleClick(TSharedRef<SWidget> OwningWidget, const FGeometry& InMyGeometry, const FPointerEvent& InMouseEvent) { return FReply::Unhandled(); }
 	virtual void OnFocusLost(const FFocusEvent& InFocusEvent) {}
+	virtual TSharedPtr<FStructOnScope> GetToolOptions() const { return nullptr; }
+	virtual void OnToolOptionsUpdated(const FPropertyChangedEvent& PropertyChangedEvent) {}
 	// ~SWidget
 
 	/**
@@ -131,6 +136,8 @@ public:
 		ToolID = InToolID;
 	}
 
+	FOnOptionsRefresh OnOptionsRefreshDelegate;
+
 protected:
 	FCurveEditorToolID ToolID;
 };
\ No newline at end of file
diff --git a/Engine/Source/Editor/CurveEditor/Public/SCurveEditorPanel.h b/Engine/Source/Editor/CurveEditor/Public/SCurveEditorPanel.h
--- a/Engine/Source/Editor/CurveEditor/Public/SCurveEditorPanel.h
+++ b/Engine/Source/Editor/CurveEditor/Public/SCurveEditorPanel.h
@@ -18,6 +18,7 @@ class ITimeSliderController;
 class SScrollBox;
 class SCurveEditorView;
 class SCurveEditorViewContainer;
+class SCurveEditorToolProperties;
 class FTabManager;
 class UCurveEditorFilterBase;
 
@@ -267,6 +268,9 @@ private:
 	/** Edit panel */
 	TSharedPtr<SCurveKeyDetailPanel> KeyDetailsView;
 
+	/** Tool options panel */
+	TSharedPtr<SCurveEditorToolProperties> ToolPropertiesPanel;
+
 	/** Map of edit UI widgets for each curve in the current selection set */
 	TMap<FCurveModelID, TSharedPtr<SWidget>> CurveToEditUI;
 
@@ -287,6 +291,9 @@ private:
 	/** Rebuild the Curve Views layout to match the currently specified View Mode. */
 	void RebuildCurveViews();
 
+	/** Reconstructs the properties widget on tool switch */
+	void OnCurveEditorToolChanged(FCurveEditorToolID InToolId);
+
 	/** The last set View Mode for this UI. */
 	ECurveEditorViewID DefaultViewID;
 
