Fix missing include to StructOnScope.h in CurveEditorTransformTool.cpp
#rnx

diff --git a/Engine/Plugins/Editor/CurveEditorTools/Source/CurveEditorTools/Private/Tools/CurveEditorTransformTool.cpp b/Engine/Plugins/Editor/CurveEditorTools/Source/CurveEditorTools/Private/Tools/CurveEditorTransformTool.cpp
--- a/Engine/Plugins/Editor/CurveEditorTools/Source/CurveEditorTools/Private/Tools/CurveEditorTransformTool.cpp
+++ b/Engine/Plugins/Editor/CurveEditorTools/Source/CurveEditorTools/Private/Tools/CurveEditorTransformTool.cpp
@@ -17,6 +17,7 @@
 #include "Containers/ArrayView.h"
 #include "CurveDataAbstraction.h"
 #include "CurveEditorSnapMetrics.h"
+#include "UObject/StructOnScope.h"
 
 #define LOCTEXT_NAMESPACE "CurveEditorToolCommands"
 
