Fix linux compile error
#rnx

diff --git a/Engine/Plugins/Editor/CurveEditorTools/Source/CurveEditorTools/Private/Tools/CurveEditorTransformTool.cpp b/Engine/Plugins/Editor/CurveEditorTools/Source/CurveEditorTools/Private/Tools/CurveEditorTransformTool.cpp
--- a/Engine/Plugins/Editor/CurveEditorTools/Source/CurveEditorTools/Private/Tools/CurveEditorTransformTool.cpp
+++ b/Engine/Plugins/Editor/CurveEditorTools/Source/CurveEditorTools/Private/Tools/CurveEditorTransformTool.cpp
@@ -698,7 +698,7 @@ void FCurveEditorTransformTool::OnDrag(const FPointerEvent& InMouseEvent, const
 			FalloffHeight = 1.0f - MouseDelta.Y / TransformWidget.Size.Y;
 			FalloffHeight = FMath::Clamp(FalloffHeight, 0.0f, 1.0f);
 		}
-		if (bTopLeft || bTopRight && TransformWidget.Size.X > SMALL_NUMBER)
+		if ((bTopLeft || bTopRight) && TransformWidget.Size.X > SMALL_NUMBER)
 		{
 			FVector2D MouseDelta = InLocalMousePosition - TransformWidget.Position;
 
