Curve Editor: Add falloff type option in transform tool, visualize falloff strength

diff --git a/Engine/Plugins/Editor/CurveEditorTools/Source/CurveEditorTools/Private/Tools/CurveEditorTransformTool.cpp b/Engine/Plugins/Editor/CurveEditorTools/Source/CurveEditorTools/Private/Tools/CurveEditorTransformTool.cpp
--- a/Engine/Plugins/Editor/CurveEditorTools/Source/CurveEditorTools/Private/Tools/CurveEditorTransformTool.cpp
+++ b/Engine/Plugins/Editor/CurveEditorTools/Source/CurveEditorTools/Private/Tools/CurveEditorTransformTool.cpp
@@ -18,6 +18,7 @@
 #include "CurveDataAbstraction.h"
 #include "CurveEditorSnapMetrics.h"
 #include "UObject/StructOnScope.h"
+#include "Algo/Transform.h"
 
 #define LOCTEXT_NAMESPACE "CurveEditorToolCommands"
 
@@ -27,6 +28,8 @@ namespace CurveEditorTransformTool
 	constexpr float EdgeAnchorWidth = 13.f;
 	constexpr float SoftSelecdtAnchorWidth = 20.f;
 	constexpr float EdgeHighlightAlpha = 0.15f;
+	constexpr float MaxFalloffOpacity = 0.6f;
+	constexpr int32 FalloffGradientSampleSize = 10;
 }
 
 void FCurveEditorTransformWidget::GetSidebarGeometry(const FGeometry& InWidgetGeometry, FGeometry& OutLeft, FGeometry& OutRight, FGeometry& OutTop, FGeometry& OutBottom) const
@@ -458,7 +461,6 @@ void FCurveEditorTransformTool::OnToolOptionsUpdated(const FPropertyChangedEvent
 		ScaleCenter.Y = 0.0f;
 		ScaleDelta.Y = CurveSpace.ValueToScreen(ToolOptions.LowerBound) - (TransformWidget.Position.Y + TransformWidget.Size.Y);
 	}
-
 	
 	const FVector2D PanelSpaceCenter = TransformWidget.Position + (TransformWidget.Size * ScaleCenter);
 	const FVector2D ChangeAmount = (ScaleDelta / TransformWidget.Size);
@@ -487,7 +489,8 @@ void FCurveEditorTransformTool::DrawMarqueeWidget(const FCurveEditorTransformWid
 
 	// Draw the inner marquee dotted rectangle line and the highlight
 	{
-		FLinearColor CenterHighlightColor = (InTransformWidget.SelectedAnchorFlags == ECurveEditorAnchorFlags::Center) ? FLinearColor::White.CopyWithNewOpacity(CurveEditorTransformTool::EdgeHighlightAlpha) : FLinearColor::Transparent;
+		FLinearColor CenterHighlightColor = (InTransformWidget.SelectedAnchorFlags == ECurveEditorAnchorFlags::Center) && !FSlateApplication::Get().GetModifierKeys().IsControlDown() ? 
+			FLinearColor::White.CopyWithNewOpacity(CurveEditorTransformTool::EdgeHighlightAlpha) : FLinearColor::Transparent;
 		FGeometry CenterGeometry;
 		InTransformWidget.GetCenterGeometry(InAllottedGeometry, CenterGeometry);
 
@@ -548,6 +551,50 @@ void FCurveEditorTransformTool::DrawMarqueeWidget(const FCurveEditorTransformWid
 		FSlateDrawElement::MakeRotatedBox(OutDrawElements, InPaintOnLayerId, RightFalloffGeometry.ToPaintGeometry(), FEditorStyle::GetBrush(TEXT("WhiteBrush")), ESlateDrawEffect::None, Rotate,
 			TOptional<FVector2D>(), FSlateDrawElement::RelativeToElement, BottomRightHighlightColor);
 		FSlateDrawElement::MakeRotatedBox(OutDrawElements, InPaintOnLayerId, RightFalloffGeometry.ToPaintGeometry(), FEditorStyle::GetBrush(TEXT("MarqueeSelection")), ESlateDrawEffect::None, Rotate);
+
+		// Draw falloff weights
+		FGeometry GradientGeometry;
+		InTransformWidget.GetCenterGeometry(InAllottedGeometry, GradientGeometry);
+
+		// Sample weights at 10 points to get an estimation of the falloff strength
+		TArray<FSlateGradientStop> GradientStops;
+		GradientStops.Reserve(CurveEditorTransformTool::FalloffGradientSampleSize * 2 + 1);
+		TArray<float> SampleVals;
+		SampleVals.SetNumUninitialized(CurveEditorTransformTool::FalloffGradientSampleSize);
+		for (int32 Index = 0; Index < CurveEditorTransformTool::FalloffGradientSampleSize; Index++)
+		{
+			SampleVals[Index] = .1f * (Index + 1);
+		}
+
+		TArray<float> InterpolatedSamples;
+		InterpolatedSamples.Reserve(CurveEditorTransformTool::FalloffGradientSampleSize);
+		Algo::Transform(SampleVals, InterpolatedSamples, [this] (float Val) { return FMath::Abs(ModifyWeightByInterpType(Val)); });
+
+		// Draw gradients, left half then right half
+		GradientStops.Emplace(FVector2D(0.f, GradientGeometry.GetLocalSize().Y), FLinearColor::Gray.CopyWithNewOpacity(0.0f));
+		for (int32 Index = 0; Index < CurveEditorTransformTool::FalloffGradientSampleSize; Index++)
+		{
+			FVector2D AnchorPos;
+			AnchorPos.X = (SampleVals[Index]) * (GradientGeometry.GetLocalSize().X * 0.5f * (1-FalloffWidth));
+			AnchorPos.Y = GradientGeometry.GetLocalSize().Y * 0.5f;
+
+			FLinearColor StopColor = FLinearColor::Gray.CopyWithNewOpacity(InterpolatedSamples[Index] * CurveEditorTransformTool::MaxFalloffOpacity);
+
+			GradientStops.Emplace(AnchorPos, StopColor);
+		}
+		for (int32 Index = 0; Index < CurveEditorTransformTool::FalloffGradientSampleSize; Index++)
+		{
+			FVector2D AnchorPos;
+			AnchorPos.X = (((SampleVals[Index] - .1f) * (1 - FalloffWidth)) + (1 + FalloffWidth)) * (GradientGeometry.GetLocalSize().X * 0.5f);
+			AnchorPos.Y = GradientGeometry.GetLocalSize().Y;
+
+			FLinearColor StopColor = FLinearColor::Gray.CopyWithNewOpacity(
+				InterpolatedSamples[CurveEditorTransformTool::FalloffGradientSampleSize - Index - 1] * CurveEditorTransformTool::MaxFalloffOpacity);
+
+			GradientStops.Emplace(AnchorPos, StopColor);
+		}
+		
+		FSlateDrawElement::MakeGradient(OutDrawElements, InPaintOnLayerId, GradientGeometry.ToPaintGeometry(), GradientStops, Orient_Vertical);
 	}
 	// draw center marker if on
 	if (FSlateApplication::Get().GetModifierKeys().IsAltDown())
@@ -1009,7 +1056,7 @@ double FCurveEditorTransformTool::GetFalloffWeight(double InputValue) const
 double FCurveEditorTransformTool::ModifyWeightByInterpType(double Value) const
 {
 	float Result = Value;
-	switch (FalloffInterpType)
+	switch (ToolOptions.FalloffInterpType)
 	{
 		case EToolTransformInterpType::Linear:
 		{
diff --git a/Engine/Plugins/Editor/CurveEditorTools/Source/CurveEditorTools/Private/Tools/CurveEditorTransformTool.h b/Engine/Plugins/Editor/CurveEditorTools/Source/CurveEditorTools/Private/Tools/CurveEditorTransformTool.h
--- a/Engine/Plugins/Editor/CurveEditorTools/Source/CurveEditorTools/Private/Tools/CurveEditorTransformTool.h
+++ b/Engine/Plugins/Editor/CurveEditorTools/Source/CurveEditorTools/Private/Tools/CurveEditorTransformTool.h
@@ -31,6 +31,7 @@ enum class ECurveEditorAnchorFlags : uint16
 	CenterScale = 0x200
 };
 
+UENUM()
 enum class EToolTransformInterpType : uint8
 {
 	Linear,
@@ -39,7 +40,7 @@ enum class EToolTransformInterpType : uint8
 	CircularIn,
 	CircularOut,
 	ExpIn,
-	ExpOut,
+	ExpOut
 };
 
 ENUM_CLASS_FLAGS(ECurveEditorAnchorFlags);
@@ -97,6 +98,10 @@ struct FTransformToolOptions
 
 	UPROPERTY(Transient, EditAnywhere, Category = ToolOptions)
 	float RightBound;
+
+	/** specifies the falloff type applied to curve selection */
+	UPROPERTY(Transient, EditAnywhere, Category = ToolOptions)
+	EToolTransformInterpType FalloffInterpType;
 };
 
 class FCurveEditorTransformTool : public ICurveEditorToolExtension
@@ -104,11 +109,12 @@ class FCurveEditorTransformTool : public ICurveEditorToolExtension
 public:
 	FCurveEditorTransformTool(TWeakPtr<FCurveEditor> InCurveEditor)
 		: WeakCurveEditor(InCurveEditor), 
-		FalloffInterpType(EToolTransformInterpType::Linear),
 		FalloffHeight(0.0f), 
 		FalloffWidth(0.0f), 
 		RelativeScaleCenter(FVector2D(0.5f, 0.5f))
-	{}
+	{
+		ToolOptions.FalloffInterpType = EToolTransformInterpType::Linear;
+	}
 
 	// ICurveEditorToolExtension Interface
 	virtual void OnPaint(const FPaintArgs& Args, const FGeometry& AllottedGeometry, const FSlateRect& MyCullingRect, FSlateWindowElementList& OutDrawElements, int32 PaintOnLayerId, const FWidgetStyle& InWidgetStyle, bool bParentEnabled) const override;
@@ -170,17 +176,15 @@ private:
 	FVector2D InitialMousePosition;
 	FCurveEditorAxisSnap::FSnapState SnappingState;
 
-	// UStruct that displays properties to be modified on screen
+	/** UStruct that displays properties to be modified on screen */
 	FTransformToolOptions ToolOptions;
 	TSharedPtr<FStructOnScope> ToolOptionsOnScope;
 
-	// helper struct for checking if tool has changed outside of own operations, i.e. middle click
+	/** helper struct for checking if tool has changed outside of own operations, i.e. middle click */
 	struct ToolBounds { FVector2D Positon, Size; };
 	ToolBounds PrevState;
 
 private:
-	/** specifies the falloff curve applied to soft selection */
-	EToolTransformInterpType FalloffInterpType;
 	/** Soft Selection Height, if 0.0 as we go up to 1.0 we go up to the top edge */
 	float FalloffHeight;
 	/** Soft Selection Width, if 0.0 we peak directly at the midpoint gradually scaling out till we reach 1.0fff */
