Curve Editor: small bug fixes for transform tool

diff --git a/Engine/Plugins/Editor/CurveEditorTools/Source/CurveEditorTools/Private/Tools/CurveEditorTransformTool.cpp b/Engine/Plugins/Editor/CurveEditorTools/Source/CurveEditorTools/Private/Tools/CurveEditorTransformTool.cpp
--- a/Engine/Plugins/Editor/CurveEditorTools/Source/CurveEditorTools/Private/Tools/CurveEditorTransformTool.cpp
+++ b/Engine/Plugins/Editor/CurveEditorTools/Source/CurveEditorTools/Private/Tools/CurveEditorTransformTool.cpp
@@ -523,18 +523,21 @@ void FCurveEditorTransformTool::DrawMarqueeWidget(const FCurveEditorTransformWid
 		FGeometry TopLeftFalloffGeometry, TopRightFalloffGeometry, LeftFalloffGeometry, RightFalloffGeometry;
 		InTransformWidget.GetFalloffGeometry(InAllottedGeometry, FalloffHeight, FalloffWidth, TopLeftFalloffGeometry, TopRightFalloffGeometry, LeftFalloffGeometry, RightFalloffGeometry);
 
-		// fixed logic to set bools correctly
-		const bool bTopLeft = InTransformWidget.SelectedAnchorFlags == ECurveEditorAnchorFlags::FalloffTopLeft || InTransformWidget.SelectedAnchorFlags == ECurveEditorAnchorFlags::FalloffTopRight;
-		const bool bTopRight = InTransformWidget.SelectedAnchorFlags == ECurveEditorAnchorFlags::FalloffTopLeft || InTransformWidget.SelectedAnchorFlags == ECurveEditorAnchorFlags::FalloffTopRight;
-		const bool bBottomLeft = InTransformWidget.SelectedAnchorFlags == ECurveEditorAnchorFlags::FalloffLeft || InTransformWidget.SelectedAnchorFlags == ECurveEditorAnchorFlags::FalloffRight;
-		const bool bBottomRight = InTransformWidget.SelectedAnchorFlags == ECurveEditorAnchorFlags::FalloffLeft || InTransformWidget.SelectedAnchorFlags == ECurveEditorAnchorFlags::FalloffRight;
+		const bool bTopLeft = (InTransformWidget.SelectedAnchorFlags & ECurveEditorAnchorFlags::FalloffTopLeft)  != ECurveEditorAnchorFlags::None || 
+			(InTransformWidget.SelectedAnchorFlags & ECurveEditorAnchorFlags::FalloffTopRight) != ECurveEditorAnchorFlags::None;
+		const bool bTopRight = (InTransformWidget.SelectedAnchorFlags & ECurveEditorAnchorFlags::FalloffTopLeft) != ECurveEditorAnchorFlags::None ||
+			(InTransformWidget.SelectedAnchorFlags & ECurveEditorAnchorFlags::FalloffTopRight) != ECurveEditorAnchorFlags::None;
+		const bool bBottomLeft = (InTransformWidget.SelectedAnchorFlags & ECurveEditorAnchorFlags::FalloffLeft) != ECurveEditorAnchorFlags::None || 
+			(InTransformWidget.SelectedAnchorFlags & ECurveEditorAnchorFlags::FalloffRight) != ECurveEditorAnchorFlags::None;
+		const bool bBottomRight = (InTransformWidget.SelectedAnchorFlags & ECurveEditorAnchorFlags::FalloffLeft) != ECurveEditorAnchorFlags::None ||
+			(InTransformWidget.SelectedAnchorFlags & ECurveEditorAnchorFlags::FalloffRight) != ECurveEditorAnchorFlags::None;
 
 		FLinearColor TopLeftHighlightColor = bTopLeft ? FLinearColor::White.CopyWithNewOpacity(CurveEditorTransformTool::EdgeHighlightAlpha) : FLinearColor::Transparent;
 		FLinearColor TopRightHighlightColor = bTopRight ? FLinearColor::White.CopyWithNewOpacity(CurveEditorTransformTool::EdgeHighlightAlpha) : FLinearColor::Transparent;
 		FLinearColor BottomLeftHighlightColor = bBottomLeft ? FLinearColor::White.CopyWithNewOpacity(CurveEditorTransformTool::EdgeHighlightAlpha) : FLinearColor::Transparent;
 		FLinearColor BottomRightHighlightColor = bBottomRight ? FLinearColor::White.CopyWithNewOpacity(CurveEditorTransformTool::EdgeHighlightAlpha) : FLinearColor::Transparent;
 
-		const float Rotate = 180.0f;
+		const float Rotate = FMath::DegreesToRadians(45.f);
 		// Top Left (Highlight, Corner Icon)
 		FSlateDrawElement::MakeRotatedBox(OutDrawElements, InPaintOnLayerId, TopLeftFalloffGeometry.ToPaintGeometry(), FEditorStyle::GetBrush(TEXT("WhiteBrush")), ESlateDrawEffect::None, Rotate,
 			TOptional<FVector2D>(), FSlateDrawElement::RelativeToElement, TopLeftHighlightColor);
diff --git a/Engine/Source/Editor/CurveEditor/Private/SCurveEditorToolProperties.cpp b/Engine/Source/Editor/CurveEditor/Private/SCurveEditorToolProperties.cpp
--- a/Engine/Source/Editor/CurveEditor/Private/SCurveEditorToolProperties.cpp
+++ b/Engine/Source/Editor/CurveEditor/Private/SCurveEditorToolProperties.cpp
@@ -59,16 +59,11 @@ void SCurveEditorToolProperties::RebuildProperties()
 	if (CurveEditor)
 	{
 		TSharedPtr<FStructOnScope> ToolOptions;
-		if (CurveEditor->GetToolExtensions().Contains(ToolId) &&
-			(ToolOptions = CurveEditor->GetToolExtensions()[ToolId]->GetToolOptions()))
+		if (CurveEditor->GetToolExtensions().Contains(ToolId))
 		{
-			DetailsView->GetWidget()->SetVisibility(EVisibility::Visible);
+			ToolOptions = CurveEditor->GetToolExtensions()[ToolId]->GetToolOptions();
 			DetailsView->SetStructureData(ToolOptions);
 		}
-		else
-		{
-			DetailsView->GetWidget()->SetVisibility(EVisibility::Hidden);
-		}
 	}
 }
 
