Curve Editor: Small transform tool modifications

diff --git a/Engine/Plugins/Editor/CurveEditorTools/Source/CurveEditorTools/Private/Tools/CurveEditorTransformTool.cpp b/Engine/Plugins/Editor/CurveEditorTools/Source/CurveEditorTools/Private/Tools/CurveEditorTransformTool.cpp
--- a/Engine/Plugins/Editor/CurveEditorTools/Source/CurveEditorTools/Private/Tools/CurveEditorTransformTool.cpp
+++ b/Engine/Plugins/Editor/CurveEditorTools/Source/CurveEditorTools/Private/Tools/CurveEditorTransformTool.cpp
@@ -24,9 +24,9 @@
 
 namespace CurveEditorTransformTool
 {
-	constexpr float ScaleCenterDiameter = 16.0f;
+	constexpr float ScaleCenterRadius = 16.0f;
 	constexpr float EdgeAnchorWidth = 13.f;
-	constexpr float SoftSelecdtAnchorWidth = 20.f;
+	constexpr float SoftSelectAnchorWidth = 20.f;
 	constexpr float EdgeHighlightAlpha = 0.15f;
 	constexpr float MaxFalloffOpacity = 0.6f;
 	constexpr int32 FalloffGradientSampleSize = 10;
@@ -71,7 +71,7 @@ void FCurveEditorTransformWidget::GetCenterGeometry(const FGeometry& InWidgetGeo
 
 void FCurveEditorTransformWidget::GetFalloffGeometry(const FGeometry& InWidgetGeometry, float FalloffHeight, float FalloffWidth, FGeometry& OutTopLeft, FGeometry& OutTopRight, FGeometry& OutLeft, FGeometry& OutRight) const
 {
-	const FVector2D CornerSize = FVector2D(CurveEditorTransformTool::SoftSelecdtAnchorWidth, CurveEditorTransformTool::SoftSelecdtAnchorWidth);
+	const FVector2D CornerSize = FVector2D(CurveEditorTransformTool::SoftSelectAnchorWidth, CurveEditorTransformTool::SoftSelectAnchorWidth);
 	const FVector2D HalfSizeOffset = FVector2D(CornerSize / 2.f);
 	const float HeightDiff = InWidgetGeometry.GetLocalSize().Y * FalloffHeight;
 	const float WidthDiff = InWidgetGeometry.GetLocalSize().X * 0.5f * FalloffWidth;
@@ -88,12 +88,19 @@ void FCurveEditorTransformWidget::GetFalloffGeometry(const FGeometry& InWidgetGe
 
 void FCurveEditorTransformWidget::GetScaleCenterGeometry(const FGeometry& InWidgetGeometry, FVector2D ScaleCenter, FGeometry& OutScaleCenterGeometry) const
 {
-	const FVector2D CenterSize = FVector2D(CurveEditorTransformTool::ScaleCenterDiameter * 2.f, CurveEditorTransformTool::ScaleCenterDiameter * 2.f);
+	const FVector2D CenterSize = FVector2D(CurveEditorTransformTool::ScaleCenterRadius * 2.f, CurveEditorTransformTool::ScaleCenterRadius * 2.f);
 	const FVector2D CenterOffset = (InWidgetGeometry.GetLocalSize() * .5f) + (BoundsSize * (ScaleCenter - 0.5f)) - (CenterSize * .5f);
 	OutScaleCenterGeometry = InWidgetGeometry.MakeChild(CenterSize, FSlateLayoutTransform(CenterOffset));
 }
 
-ECurveEditorAnchorFlags FCurveEditorTransformWidget::GetAnchorFlagsForMousePosition(const FGeometry& InWidgetGeometry, float  FalloffHeight, float FalloffWidth, const FVector2D& RelativeCenterScale, const FVector2D& InMouseScreenPosition) const
+void FCurveEditorTransformWidget::GetCenterIndicatorGeometry(const FGeometry& InWidgetGeometry, FGeometry& OutCenterIndicatorGeometry) const
+{
+	const FVector2D IndicatorSize = FVector2D(CurveEditorTransformTool::ScaleCenterRadius * 2.f, CurveEditorTransformTool::ScaleCenterRadius * 2.f);
+	const FVector2D IndicatorOffset = (InWidgetGeometry.GetLocalSize() * .5f) - (IndicatorSize * 0.5f);
+	OutCenterIndicatorGeometry = InWidgetGeometry.MakeChild(IndicatorSize, FSlateLayoutTransform(IndicatorOffset));
+}
+
+ECurveEditorAnchorFlags FCurveEditorTransformWidget::GetAnchorFlagsForMousePosition(const FGeometry& InWidgetGeometry, float  FalloffHeight, float FalloffWidth, const FVector2D& RelativeScaleCenter, const FVector2D& InMouseScreenPosition) const
 {
 	// We store a geometry to represent each different region, updated on Tick. We check if the mouse
 	// overlaps a region and update the selection anchors depending on which region you're hovering in.
@@ -127,16 +134,12 @@ ECurveEditorAnchorFlags FCurveEditorTransformWidget::GetAnchorFlagsForMousePosit
 		}
 	}
 
-	if (FSlateApplication::Get().GetModifierKeys().IsAltDown())
-	{
-		FGeometry CenterScaleGeometry;
-		GetScaleCenterGeometry(InWidgetGeometry, RelativeCenterScale, CenterScaleGeometry);
+	FGeometry ScaleCenterGeometry;
+	GetScaleCenterGeometry(InWidgetGeometry, RelativeScaleCenter, ScaleCenterGeometry);
 
-		if (CenterScaleGeometry.IsUnderLocation(InMouseScreenPosition))
-		{
-			OutFlags |= ECurveEditorAnchorFlags::CenterScale;
-			return OutFlags;
-		}
+	if (ScaleCenterGeometry.IsUnderLocation(InMouseScreenPosition))
+	{
+		OutFlags |= ECurveEditorAnchorFlags::ScaleCenter;
 	}
 
 	FGeometry LeftSidebarGeometry, RightSidebarGeometry, TopSidebarGeometry, BottomSidebarGeometry;
@@ -323,6 +326,9 @@ void FCurveEditorTransformTool::UpdateToolOptions()
 		ToolOptions.UpperBound = CurveSpace.ScreenToValue(TransformWidget.BoundsPosition.Y);
 		ToolOptions.RightBound = CurveSpace.ScreenToSeconds(TransformWidget.BoundsPosition.X + TransformWidget.BoundsSize.X);
 		ToolOptions.LowerBound = CurveSpace.ScreenToValue(TransformWidget.BoundsPosition.Y + TransformWidget.BoundsSize.Y);
+
+		ToolOptions.ScaleCenterX = CurveSpace.ScreenToSeconds(TransformWidget.BoundsPosition.X + TransformWidget.BoundsSize.X * DisplayRelativeScaleCenter.X);
+		ToolOptions.ScaleCenterY = CurveSpace.ScreenToValue(TransformWidget.BoundsPosition.Y + TransformWidget.BoundsSize.Y * DisplayRelativeScaleCenter.Y);
 	}
 }
 
@@ -346,12 +352,22 @@ FReply FCurveEditorTransformTool::OnMouseButtonDown(TSharedRef<SWidget> OwningWi
 
 FReply FCurveEditorTransformTool::OnMouseMove(TSharedRef<SWidget> OwningWidget, const FGeometry& MyGeometry, const FPointerEvent& MouseEvent)
 {
+	TSharedPtr<FCurveEditor> CurveEditor = WeakCurveEditor.Pin();
+	if (!CurveEditor)
+	{
+		return FReply::Unhandled();
+	}
+
 	// Update the Hover State of the widget.
 	if (!DelayedDrag.IsSet())
 	{
 		FGeometry WidgetGeo = TransformWidget.MakeGeometry(MyGeometry);
 		ECurveEditorAnchorFlags HitWidgetFlags = TransformWidget.GetAnchorFlagsForMousePosition(WidgetGeo, FalloffHeight, FalloffWidth, RelativeScaleCenter, MouseEvent.GetScreenSpacePosition());
-		TransformWidget.SelectedAnchorFlags = HitWidgetFlags;
+		if (CurveEditor->GetSelection().Count() <= 1)
+		{
+			HitWidgetFlags ^= ECurveEditorAnchorFlags::ScaleCenter;
+		}
+		TransformWidget.DisplayAnchorFlags = TransformWidget.SelectedAnchorFlags = HitWidgetFlags;
 	}
 
 	if (DelayedDrag.IsSet())
@@ -461,7 +477,20 @@ void FCurveEditorTransformTool::OnToolOptionsUpdated(const FPropertyChangedEvent
 		ScaleCenter.Y = 0.0f;
 		ScaleDelta.Y = CurveSpace.ValueToScreen(ToolOptions.LowerBound) - (TransformWidget.Position.Y + TransformWidget.Size.Y);
 	}
+	else if (PropertyChangedEvent.GetPropertyName().IsEqual(GET_MEMBER_NAME_CHECKED(FTransformToolOptions, ScaleCenterX)))
+	{
+		const double ViewSpaceX = CurveSpace.SecondsToScreen(ToolOptions.ScaleCenterX);
+		const double ViewSpaceXDelta = ViewSpaceX - TransformWidget.BoundsPosition.X;
+		RelativeScaleCenter.X = ViewSpaceXDelta / TransformWidget.BoundsSize.X;
+	}
+	else if (PropertyChangedEvent.GetPropertyName().IsEqual(GET_MEMBER_NAME_CHECKED(FTransformToolOptions, ScaleCenterY)))
+	{
+		const double ViewSpaceY = CurveSpace.ValueToScreen(ToolOptions.ScaleCenterY);
+		const double ViewSpaceYDelta = ViewSpaceY - TransformWidget.BoundsPosition.Y;
+		RelativeScaleCenter.Y = ViewSpaceYDelta / TransformWidget.BoundsSize.Y;
+	}
 	
+	DisplayRelativeScaleCenter = RelativeScaleCenter;
 	const FVector2D PanelSpaceCenter = TransformWidget.Position + (TransformWidget.Size * ScaleCenter);
 	const FVector2D ChangeAmount = (ScaleDelta / TransformWidget.Size);
 
@@ -471,6 +500,7 @@ void FCurveEditorTransformTool::OnToolOptionsUpdated(const FPropertyChangedEvent
 
 	PrevState.Positon = TransformWidget.Position;
 	PrevState.Size = TransformWidget.Size;
+	ActiveTransaction.Reset();
 }
 
 void FCurveEditorTransformTool::OnPaint(const FPaintArgs& Args, const FGeometry& AllottedGeometry, const FSlateRect& MyCullingRect, FSlateWindowElementList& OutDrawElements, int32 PaintOnLayerId, const FWidgetStyle& InWidgetStyle, bool bParentEnabled) const
@@ -482,14 +512,18 @@ void FCurveEditorTransformTool::OnPaint(const FPaintArgs& Args, const FGeometry&
 
 void FCurveEditorTransformTool::DrawMarqueeWidget(const FCurveEditorTransformWidget& InTransformWidget, const FPaintArgs& InArgs, const FGeometry& InAllottedGeometry, const FSlateRect& InMyCullingRect, FSlateWindowElementList& OutDrawElements, const int32 InPaintOnLayerId, const FWidgetStyle& InWidgetStyle, const bool bInParentEnabled) const
 {
-	if (!InTransformWidget.Visible)
+	TSharedPtr<FCurveEditor> CurveEditor = WeakCurveEditor.Pin();
+	if (!InTransformWidget.Visible || !CurveEditor)
 	{
 		return;
 	}
 
 	// Draw the inner marquee dotted rectangle line and the highlight
 	{
-		FLinearColor CenterHighlightColor = (InTransformWidget.SelectedAnchorFlags == ECurveEditorAnchorFlags::Center) && !FSlateApplication::Get().GetModifierKeys().IsControlDown() ? 
+		const bool bCenterHovered = (InTransformWidget.DisplayAnchorFlags & ECurveEditorAnchorFlags::Center) != ECurveEditorAnchorFlags::None;
+		const bool bScaleCenterNotShowing = ((InTransformWidget.DisplayAnchorFlags & ECurveEditorAnchorFlags::ScaleCenter) == ECurveEditorAnchorFlags::None || CurveEditor->GetSelection().Count() == 1);
+		const bool bFalloffOn = FSlateApplication::Get().GetModifierKeys().IsControlDown();
+		FLinearColor CenterHighlightColor = bCenterHovered && bScaleCenterNotShowing && !bFalloffOn ?
 			FLinearColor::White.CopyWithNewOpacity(CurveEditorTransformTool::EdgeHighlightAlpha) : FLinearColor::Transparent;
 		FGeometry CenterGeometry;
 		InTransformWidget.GetCenterGeometry(InAllottedGeometry, CenterGeometry);
@@ -500,10 +534,10 @@ void FCurveEditorTransformTool::DrawMarqueeWidget(const FCurveEditorTransformWid
 
 	// Draw edge highlight regions on mouse hover
 	{
-		FLinearColor LeftEdgeHighlightColor =	(InTransformWidget.SelectedAnchorFlags == ECurveEditorAnchorFlags::Left)	? FLinearColor::White.CopyWithNewOpacity(CurveEditorTransformTool::EdgeHighlightAlpha) : FLinearColor::Transparent;
-		FLinearColor RightEdgeHighlightColor =	(InTransformWidget.SelectedAnchorFlags == ECurveEditorAnchorFlags::Right)	? FLinearColor::White.CopyWithNewOpacity(CurveEditorTransformTool::EdgeHighlightAlpha) : FLinearColor::Transparent;
-		FLinearColor TopEdgeHighlightColor =	(InTransformWidget.SelectedAnchorFlags == ECurveEditorAnchorFlags::Top)		? FLinearColor::White.CopyWithNewOpacity(CurveEditorTransformTool::EdgeHighlightAlpha) : FLinearColor::Transparent;
-		FLinearColor BottomEdgeHighlightColor =	(InTransformWidget.SelectedAnchorFlags == ECurveEditorAnchorFlags::Bottom)	? FLinearColor::White.CopyWithNewOpacity(CurveEditorTransformTool::EdgeHighlightAlpha) : FLinearColor::Transparent;
+		FLinearColor LeftEdgeHighlightColor =	(InTransformWidget.DisplayAnchorFlags == ECurveEditorAnchorFlags::Left)	? FLinearColor::White.CopyWithNewOpacity(CurveEditorTransformTool::EdgeHighlightAlpha) : FLinearColor::Transparent;
+		FLinearColor RightEdgeHighlightColor =	(InTransformWidget.DisplayAnchorFlags == ECurveEditorAnchorFlags::Right)	? FLinearColor::White.CopyWithNewOpacity(CurveEditorTransformTool::EdgeHighlightAlpha) : FLinearColor::Transparent;
+		FLinearColor TopEdgeHighlightColor =	(InTransformWidget.DisplayAnchorFlags == ECurveEditorAnchorFlags::Top)		? FLinearColor::White.CopyWithNewOpacity(CurveEditorTransformTool::EdgeHighlightAlpha) : FLinearColor::Transparent;
+		FLinearColor BottomEdgeHighlightColor =	(InTransformWidget.DisplayAnchorFlags == ECurveEditorAnchorFlags::Bottom)	? FLinearColor::White.CopyWithNewOpacity(CurveEditorTransformTool::EdgeHighlightAlpha) : FLinearColor::Transparent;
 
 		FGeometry LeftSidebarGeometry, RightSidebarGeometry, TopSidebarGeometry, BottomSidebarGeometry;
 		InTransformWidget.GetSidebarGeometry(InAllottedGeometry, LeftSidebarGeometry, RightSidebarGeometry, TopSidebarGeometry, BottomSidebarGeometry);
@@ -523,14 +557,14 @@ void FCurveEditorTransformTool::DrawMarqueeWidget(const FCurveEditorTransformWid
 		FGeometry TopLeftFalloffGeometry, TopRightFalloffGeometry, LeftFalloffGeometry, RightFalloffGeometry;
 		InTransformWidget.GetFalloffGeometry(InAllottedGeometry, FalloffHeight, FalloffWidth, TopLeftFalloffGeometry, TopRightFalloffGeometry, LeftFalloffGeometry, RightFalloffGeometry);
 
-		const bool bTopLeft = (InTransformWidget.SelectedAnchorFlags & ECurveEditorAnchorFlags::FalloffTopLeft)  != ECurveEditorAnchorFlags::None || 
-			(InTransformWidget.SelectedAnchorFlags & ECurveEditorAnchorFlags::FalloffTopRight) != ECurveEditorAnchorFlags::None;
-		const bool bTopRight = (InTransformWidget.SelectedAnchorFlags & ECurveEditorAnchorFlags::FalloffTopLeft) != ECurveEditorAnchorFlags::None ||
-			(InTransformWidget.SelectedAnchorFlags & ECurveEditorAnchorFlags::FalloffTopRight) != ECurveEditorAnchorFlags::None;
-		const bool bBottomLeft = (InTransformWidget.SelectedAnchorFlags & ECurveEditorAnchorFlags::FalloffLeft) != ECurveEditorAnchorFlags::None || 
-			(InTransformWidget.SelectedAnchorFlags & ECurveEditorAnchorFlags::FalloffRight) != ECurveEditorAnchorFlags::None;
-		const bool bBottomRight = (InTransformWidget.SelectedAnchorFlags & ECurveEditorAnchorFlags::FalloffLeft) != ECurveEditorAnchorFlags::None ||
-			(InTransformWidget.SelectedAnchorFlags & ECurveEditorAnchorFlags::FalloffRight) != ECurveEditorAnchorFlags::None;
+		const bool bTopLeft = (InTransformWidget.DisplayAnchorFlags & ECurveEditorAnchorFlags::FalloffTopLeft)  != ECurveEditorAnchorFlags::None || 
+			(InTransformWidget.DisplayAnchorFlags & ECurveEditorAnchorFlags::FalloffTopRight) != ECurveEditorAnchorFlags::None;
+		const bool bTopRight = (InTransformWidget.DisplayAnchorFlags & ECurveEditorAnchorFlags::FalloffTopLeft) != ECurveEditorAnchorFlags::None ||
+			(InTransformWidget.DisplayAnchorFlags & ECurveEditorAnchorFlags::FalloffTopRight) != ECurveEditorAnchorFlags::None;
+		const bool bBottomLeft = (InTransformWidget.DisplayAnchorFlags & ECurveEditorAnchorFlags::FalloffLeft) != ECurveEditorAnchorFlags::None || 
+			(InTransformWidget.DisplayAnchorFlags & ECurveEditorAnchorFlags::FalloffRight) != ECurveEditorAnchorFlags::None;
+		const bool bBottomRight = (InTransformWidget.DisplayAnchorFlags & ECurveEditorAnchorFlags::FalloffLeft) != ECurveEditorAnchorFlags::None ||
+			(InTransformWidget.DisplayAnchorFlags & ECurveEditorAnchorFlags::FalloffRight) != ECurveEditorAnchorFlags::None;
 
 		FLinearColor TopLeftHighlightColor = bTopLeft ? FLinearColor::White.CopyWithNewOpacity(CurveEditorTransformTool::EdgeHighlightAlpha) : FLinearColor::Transparent;
 		FLinearColor TopRightHighlightColor = bTopRight ? FLinearColor::White.CopyWithNewOpacity(CurveEditorTransformTool::EdgeHighlightAlpha) : FLinearColor::Transparent;
@@ -574,14 +608,13 @@ void FCurveEditorTransformTool::DrawMarqueeWidget(const FCurveEditorTransformWid
 		Algo::Transform(SampleVals, InterpolatedSamples, [this] (float Val) { return FMath::Abs(ModifyWeightByInterpType(Val)); });
 
 		// Draw gradients, left half then right half
-		GradientStops.Emplace(FVector2D(0.f, GradientGeometry.GetLocalSize().Y), FLinearColor::Gray.CopyWithNewOpacity(0.0f));
 		for (int32 Index = 0; Index < CurveEditorTransformTool::FalloffGradientSampleSize; Index++)
 		{
 			FVector2D AnchorPos;
 			AnchorPos.X = (SampleVals[Index]) * (GradientGeometry.GetLocalSize().X * 0.5f * (1-FalloffWidth));
 			AnchorPos.Y = GradientGeometry.GetLocalSize().Y * 0.5f;
 
-			FLinearColor StopColor = FLinearColor::Gray.CopyWithNewOpacity(InterpolatedSamples[Index] * CurveEditorTransformTool::MaxFalloffOpacity);
+			FLinearColor StopColor = FLinearColor::Gray.CopyWithNewOpacity((FalloffHeight + (1 - FalloffHeight) * InterpolatedSamples[Index]) * CurveEditorTransformTool::MaxFalloffOpacity);
 
 			GradientStops.Emplace(AnchorPos, StopColor);
 		}
@@ -592,7 +625,7 @@ void FCurveEditorTransformTool::DrawMarqueeWidget(const FCurveEditorTransformWid
 			AnchorPos.Y = GradientGeometry.GetLocalSize().Y;
 
 			FLinearColor StopColor = FLinearColor::Gray.CopyWithNewOpacity(
-				InterpolatedSamples[CurveEditorTransformTool::FalloffGradientSampleSize - Index - 1] * CurveEditorTransformTool::MaxFalloffOpacity);
+				(FalloffHeight + (1 - FalloffHeight) * InterpolatedSamples[CurveEditorTransformTool::FalloffGradientSampleSize - Index - 1]) * CurveEditorTransformTool::MaxFalloffOpacity);
 
 			GradientStops.Emplace(AnchorPos, StopColor);
 		}
@@ -600,37 +633,46 @@ void FCurveEditorTransformTool::DrawMarqueeWidget(const FCurveEditorTransformWid
 		FSlateDrawElement::MakeGradient(OutDrawElements, InPaintOnLayerId, GradientGeometry.ToPaintGeometry(), GradientStops, Orient_Vertical);
 	}
 	// draw center marker if on
-	if (FSlateApplication::Get().GetModifierKeys().IsAltDown())
+	if (CurveEditor->GetSelection().Count() != 1)
 	{
 		FGeometry ScaleCenterGeometry;
-		InTransformWidget.GetScaleCenterGeometry(InAllottedGeometry, RelativeScaleCenter, ScaleCenterGeometry);
+		InTransformWidget.GetScaleCenterGeometry(InAllottedGeometry, DisplayRelativeScaleCenter, ScaleCenterGeometry);
 		
-		const bool bScaleCenter = InTransformWidget.SelectedAnchorFlags == ECurveEditorAnchorFlags::CenterScale;
+		const bool bScaleCenter = (InTransformWidget.DisplayAnchorFlags & ECurveEditorAnchorFlags::ScaleCenter) != ECurveEditorAnchorFlags::None;
 
 		FLinearColor ScaleCenterHighlightColor = bScaleCenter ? FLinearColor::White.CopyWithNewOpacity(CurveEditorTransformTool::EdgeHighlightAlpha) : FLinearColor::Transparent;
 		const float HighlightThickness = 7.f;
 
 		// draw center scale icon lines
-		FVector2D CircleOffset = FVector2D(CurveEditorTransformTool::ScaleCenterDiameter, CurveEditorTransformTool::ScaleCenterDiameter);
-		FSlateDrawElement::MakeSpline(OutDrawElements, InPaintOnLayerId, ScaleCenterGeometry.ToPaintGeometry(), FVector2D(-CurveEditorTransformTool::ScaleCenterDiameter, 0.f) + CircleOffset, FVector2D(CurveEditorTransformTool::ScaleCenterDiameter, 0.f), 
-			FVector2D(0.f, -CurveEditorTransformTool::ScaleCenterDiameter) + CircleOffset, FVector2D(0.f, -CurveEditorTransformTool::ScaleCenterDiameter), 1.f);
-		FSlateDrawElement::MakeSpline(OutDrawElements, InPaintOnLayerId, ScaleCenterGeometry.ToPaintGeometry(), FVector2D(CurveEditorTransformTool::ScaleCenterDiameter, 0.f) + CircleOffset, FVector2D(-CurveEditorTransformTool::ScaleCenterDiameter, 0.f),
-			FVector2D(0.f, CurveEditorTransformTool::ScaleCenterDiameter) + CircleOffset, FVector2D(0.f, CurveEditorTransformTool::ScaleCenterDiameter), 1.f);
-		FSlateDrawElement::MakeSpline(OutDrawElements, InPaintOnLayerId, ScaleCenterGeometry.ToPaintGeometry(), FVector2D(0.f, CurveEditorTransformTool::ScaleCenterDiameter) + CircleOffset, FVector2D(0.f, -CurveEditorTransformTool::ScaleCenterDiameter),
-			FVector2D(-CurveEditorTransformTool::ScaleCenterDiameter, 0.f) + CircleOffset, FVector2D(-CurveEditorTransformTool::ScaleCenterDiameter, 0.f), 1.f);
-		FSlateDrawElement::MakeSpline(OutDrawElements, InPaintOnLayerId, ScaleCenterGeometry.ToPaintGeometry(), FVector2D(0.f, -CurveEditorTransformTool::ScaleCenterDiameter) + CircleOffset, FVector2D(0.f, CurveEditorTransformTool::ScaleCenterDiameter),
-			FVector2D(CurveEditorTransformTool::ScaleCenterDiameter, 0.f) + CircleOffset, FVector2D(CurveEditorTransformTool::ScaleCenterDiameter, 0.f), 1.f);
+		const double TangentMul = 2.f;
+		FVector2D CircleOffset = FVector2D(CurveEditorTransformTool::ScaleCenterRadius, CurveEditorTransformTool::ScaleCenterRadius);
+		FSlateDrawElement::MakeSpline(OutDrawElements, InPaintOnLayerId, ScaleCenterGeometry.ToPaintGeometry(), FVector2D(-CurveEditorTransformTool::ScaleCenterRadius, 0.f) + CircleOffset, FVector2D(CurveEditorTransformTool::ScaleCenterRadius, 0.f) * TangentMul,
+			FVector2D(0.f, -CurveEditorTransformTool::ScaleCenterRadius) + CircleOffset, FVector2D(0.f, -CurveEditorTransformTool::ScaleCenterRadius) * TangentMul, 1.f);
+		FSlateDrawElement::MakeSpline(OutDrawElements, InPaintOnLayerId, ScaleCenterGeometry.ToPaintGeometry(), FVector2D(CurveEditorTransformTool::ScaleCenterRadius, 0.f) + CircleOffset, FVector2D(-CurveEditorTransformTool::ScaleCenterRadius, 0.f) * TangentMul,
+			FVector2D(0.f, CurveEditorTransformTool::ScaleCenterRadius) + CircleOffset, FVector2D(0.f, CurveEditorTransformTool::ScaleCenterRadius) * TangentMul, 1.f);
+		FSlateDrawElement::MakeSpline(OutDrawElements, InPaintOnLayerId, ScaleCenterGeometry.ToPaintGeometry(), FVector2D(0.f, CurveEditorTransformTool::ScaleCenterRadius) + CircleOffset, FVector2D(0.f, -CurveEditorTransformTool::ScaleCenterRadius) * TangentMul,
+			FVector2D(-CurveEditorTransformTool::ScaleCenterRadius, 0.f) + CircleOffset, FVector2D(-CurveEditorTransformTool::ScaleCenterRadius, 0.f) * TangentMul, 1.f);
+		FSlateDrawElement::MakeSpline(OutDrawElements, InPaintOnLayerId, ScaleCenterGeometry.ToPaintGeometry(), FVector2D(0.f, -CurveEditorTransformTool::ScaleCenterRadius) + CircleOffset, FVector2D(0.f, CurveEditorTransformTool::ScaleCenterRadius) * TangentMul,
+			FVector2D(CurveEditorTransformTool::ScaleCenterRadius, 0.f) + CircleOffset, FVector2D(CurveEditorTransformTool::ScaleCenterRadius, 0.f) * TangentMul, 1.f);
 
 		// draw center scale icon highlight
-		FSlateDrawElement::MakeSpline(OutDrawElements, InPaintOnLayerId, ScaleCenterGeometry.ToPaintGeometry(), FVector2D(-CurveEditorTransformTool::ScaleCenterDiameter, 0.f) + CircleOffset, FVector2D(CurveEditorTransformTool::ScaleCenterDiameter, 0.f),
-			FVector2D(0.f, -CurveEditorTransformTool::ScaleCenterDiameter) + CircleOffset, FVector2D(0.f, -CurveEditorTransformTool::ScaleCenterDiameter), HighlightThickness, ESlateDrawEffect::None, ScaleCenterHighlightColor);
-		FSlateDrawElement::MakeSpline(OutDrawElements, InPaintOnLayerId, ScaleCenterGeometry.ToPaintGeometry(), FVector2D(CurveEditorTransformTool::ScaleCenterDiameter, 0.f) + CircleOffset, FVector2D(-CurveEditorTransformTool::ScaleCenterDiameter, 0.f),
-			FVector2D(0.f, CurveEditorTransformTool::ScaleCenterDiameter) + CircleOffset, FVector2D(0.f, CurveEditorTransformTool::ScaleCenterDiameter), HighlightThickness, ESlateDrawEffect::None, ScaleCenterHighlightColor);
-		FSlateDrawElement::MakeSpline(OutDrawElements, InPaintOnLayerId, ScaleCenterGeometry.ToPaintGeometry(), FVector2D(0.f, CurveEditorTransformTool::ScaleCenterDiameter) + CircleOffset, FVector2D(0.f, -CurveEditorTransformTool::ScaleCenterDiameter),
-			FVector2D(-CurveEditorTransformTool::ScaleCenterDiameter, 0.f) + CircleOffset, FVector2D(-CurveEditorTransformTool::ScaleCenterDiameter, 0.f), HighlightThickness, ESlateDrawEffect::None, ScaleCenterHighlightColor);
-		FSlateDrawElement::MakeSpline(OutDrawElements, InPaintOnLayerId, ScaleCenterGeometry.ToPaintGeometry(), FVector2D(0.f, -CurveEditorTransformTool::ScaleCenterDiameter) + CircleOffset, FVector2D(0.f, CurveEditorTransformTool::ScaleCenterDiameter),
-			FVector2D(CurveEditorTransformTool::ScaleCenterDiameter, 0.f) + CircleOffset, FVector2D(CurveEditorTransformTool::ScaleCenterDiameter, 0.f), HighlightThickness, ESlateDrawEffect::None, ScaleCenterHighlightColor);
+		FSlateDrawElement::MakeSpline(OutDrawElements, InPaintOnLayerId, ScaleCenterGeometry.ToPaintGeometry(), FVector2D(-CurveEditorTransformTool::ScaleCenterRadius, 0.f) + CircleOffset, FVector2D(CurveEditorTransformTool::ScaleCenterRadius, 0.f) * TangentMul,
+			FVector2D(0.f, -CurveEditorTransformTool::ScaleCenterRadius) + CircleOffset, FVector2D(0.f, -CurveEditorTransformTool::ScaleCenterRadius) * TangentMul, HighlightThickness, ESlateDrawEffect::None, ScaleCenterHighlightColor);
+		FSlateDrawElement::MakeSpline(OutDrawElements, InPaintOnLayerId, ScaleCenterGeometry.ToPaintGeometry(), FVector2D(CurveEditorTransformTool::ScaleCenterRadius, 0.f) + CircleOffset, FVector2D(-CurveEditorTransformTool::ScaleCenterRadius, 0.f) * TangentMul,
+			FVector2D(0.f, CurveEditorTransformTool::ScaleCenterRadius) + CircleOffset, FVector2D(0.f, CurveEditorTransformTool::ScaleCenterRadius) * TangentMul, HighlightThickness, ESlateDrawEffect::None, ScaleCenterHighlightColor);
+		FSlateDrawElement::MakeSpline(OutDrawElements, InPaintOnLayerId, ScaleCenterGeometry.ToPaintGeometry(), FVector2D(0.f, CurveEditorTransformTool::ScaleCenterRadius) + CircleOffset, FVector2D(0.f, -CurveEditorTransformTool::ScaleCenterRadius) * TangentMul,
+			FVector2D(-CurveEditorTransformTool::ScaleCenterRadius, 0.f) + CircleOffset, FVector2D(-CurveEditorTransformTool::ScaleCenterRadius, 0.f) * TangentMul, HighlightThickness, ESlateDrawEffect::None, ScaleCenterHighlightColor);
+		FSlateDrawElement::MakeSpline(OutDrawElements, InPaintOnLayerId, ScaleCenterGeometry.ToPaintGeometry(), FVector2D(0.f, -CurveEditorTransformTool::ScaleCenterRadius) + CircleOffset, FVector2D(0.f, CurveEditorTransformTool::ScaleCenterRadius) * TangentMul,
+			FVector2D(CurveEditorTransformTool::ScaleCenterRadius, 0.f) + CircleOffset, FVector2D(CurveEditorTransformTool::ScaleCenterRadius, 0.f) * TangentMul, HighlightThickness, ESlateDrawEffect::None, ScaleCenterHighlightColor);
+
+		if (bScaleCenter)
+		{
+			FGeometry CenterIndicatorGeometry;
+			InTransformWidget.GetCenterIndicatorGeometry(InAllottedGeometry, CenterIndicatorGeometry);
+			FLinearColor CenterIndicatorHighlightColor = FLinearColor::White.CopyWithNewOpacity(CurveEditorTransformTool::EdgeHighlightAlpha);
 
+			FSlateDrawElement::MakeBox(OutDrawElements, InPaintOnLayerId, CenterIndicatorGeometry.ToPaintGeometry(), FEditorStyle::GetBrush(TEXT("WhiteBrush")), ESlateDrawEffect::None, CenterIndicatorHighlightColor);
+		}
 	}
 
 	// Draw the four corners + highlights
@@ -638,10 +680,10 @@ void FCurveEditorTransformTool::DrawMarqueeWidget(const FCurveEditorTransformWid
 		FGeometry TopLeftCornerGeometry, TopRightCornerGeometry, BottomLeftCornerGeometry, BottomRightCornerGeometry;
 		InTransformWidget.GetCornerGeometry(InAllottedGeometry, TopLeftCornerGeometry, TopRightCornerGeometry, BottomLeftCornerGeometry, BottomRightCornerGeometry);
 
-		const bool bTopLeft =		InTransformWidget.SelectedAnchorFlags == (ECurveEditorAnchorFlags::Top		| ECurveEditorAnchorFlags::Left );
-		const bool bTopRight =		InTransformWidget.SelectedAnchorFlags == (ECurveEditorAnchorFlags::Top		| ECurveEditorAnchorFlags::Right);
-		const bool bBottomLeft =	InTransformWidget.SelectedAnchorFlags == (ECurveEditorAnchorFlags::Bottom	| ECurveEditorAnchorFlags::Left );
-		const bool bBottomRight =	InTransformWidget.SelectedAnchorFlags == (ECurveEditorAnchorFlags::Bottom	| ECurveEditorAnchorFlags::Right);
+		const bool bTopLeft =		InTransformWidget.DisplayAnchorFlags == (ECurveEditorAnchorFlags::Top		| ECurveEditorAnchorFlags::Left );
+		const bool bTopRight =		InTransformWidget.DisplayAnchorFlags == (ECurveEditorAnchorFlags::Top		| ECurveEditorAnchorFlags::Right);
+		const bool bBottomLeft =	InTransformWidget.DisplayAnchorFlags == (ECurveEditorAnchorFlags::Bottom	| ECurveEditorAnchorFlags::Left );
+		const bool bBottomRight =	InTransformWidget.DisplayAnchorFlags == (ECurveEditorAnchorFlags::Bottom	| ECurveEditorAnchorFlags::Right);
 
 		FLinearColor TopLeftHighlightColor		= bTopLeft		? FLinearColor::White.CopyWithNewOpacity(CurveEditorTransformTool::EdgeHighlightAlpha) : FLinearColor::Transparent;
 		FLinearColor TopRightHighlightColor		= bTopRight		? FLinearColor::White.CopyWithNewOpacity(CurveEditorTransformTool::EdgeHighlightAlpha) : FLinearColor::Transparent;
@@ -735,7 +777,6 @@ void FCurveEditorTransformTool::OnDrag(const FPointerEvent& InMouseEvent, const
 
 	FSlateLayoutTransform ContainerToAbsolute = CurveEditor->GetPanel()->GetViewContainerGeometry().GetAccumulatedLayoutTransform().Inverse();
 	const bool bFalloffOn = InMouseEvent.IsControlDown();
-	const bool bScaleCenterOn = InMouseEvent.IsAltDown();
 	if (bFalloffOn && (TransformWidget.SelectedAnchorFlags & ECurveEditorAnchorFlags::FalloffAll) != ECurveEditorAnchorFlags::None)
 	{
 		const bool bTopLeft = ((TransformWidget.SelectedAnchorFlags & ECurveEditorAnchorFlags::FalloffTopLeft) != ECurveEditorAnchorFlags::None);
@@ -764,10 +805,78 @@ void FCurveEditorTransformTool::OnDrag(const FPointerEvent& InMouseEvent, const
 			}
 		}
 	}
-	else if (bScaleCenterOn && (TransformWidget.SelectedAnchorFlags & ECurveEditorAnchorFlags::CenterScale) != ECurveEditorAnchorFlags::None)
+	else if ((TransformWidget.SelectedAnchorFlags & ECurveEditorAnchorFlags::ScaleCenter) != ECurveEditorAnchorFlags::None && CurveEditor->GetSelection().Count() != 1)
 	{
 		const FVector2D MouseDelta = InLocalMousePosition - TransformWidget.BoundsPosition;
 		RelativeScaleCenter = MouseDelta / TransformWidget.BoundsSize;
+
+		const float SnapThreshold = 20.f;
+
+		// Snap to keys
+		{
+			FVector2D ClosestInRange;
+			double ClosestInRangeDist = TNumericLimits<double>::Max();
+
+			for (const TPair<FCurveModelID, FKeyHandleSet>& CurveKeys : CurveEditor->GetSelection().GetAll())
+			{
+				TArray<FKeyPosition> KeyPositions;
+				KeyPositions.SetNumUninitialized(CurveKeys.Value.Num());
+				CurveEditor->GetCurves()[CurveKeys.Key]->GetKeyPositions(CurveKeys.Value.AsArray(), KeyPositions);
+
+				const SCurveEditorView* View = CurveEditor->FindFirstInteractiveView(CurveKeys.Key);
+				if (!View)
+				{
+					return;
+				}
+
+				FCurveEditorScreenSpace CurveSpace = View->GetCurveSpace(CurveKeys.Key);
+
+				for (const FKeyPosition& KeyPos : KeyPositions)
+				{
+					const FVector2D ViewSpacePosition = FVector2D(CurveSpace.SecondsToScreen(KeyPos.InputValue), CurveSpace.ValueToScreen(KeyPos.OutputValue));
+					const double KeyDistance = FVector2D::Distance(ViewSpacePosition, InLocalMousePosition);
+					if (KeyDistance < ClosestInRangeDist)
+					{
+						ClosestInRangeDist = KeyDistance;
+						ClosestInRange = ViewSpacePosition;
+					}
+				}
+			}
+
+			if (ClosestInRangeDist < SnapThreshold)
+			{
+				const FVector2D ViewSpaceDelta = ClosestInRange - TransformWidget.BoundsPosition;
+				RelativeScaleCenter = ViewSpaceDelta / TransformWidget.BoundsSize;
+			}
+		}
+
+		// Snap to edges or corners
+		{
+			if (FMath::Abs(InLocalMousePosition.Y - (TransformWidget.Position.Y)) < SnapThreshold)
+			{
+				RelativeScaleCenter.Y = 0.f;
+			}
+			if (FMath::Abs(InLocalMousePosition.Y - (TransformWidget.Position.Y + TransformWidget.Size.Y)) < SnapThreshold)
+			{
+				RelativeScaleCenter.Y = 1.f;
+			}
+			if (FMath::Abs(InLocalMousePosition.X - (TransformWidget.Position.X)) < SnapThreshold)
+			{
+				RelativeScaleCenter.X = 0.f;
+			}
+			if (FMath::Abs(InLocalMousePosition.X - (TransformWidget.Position.X + TransformWidget.Size.X)) < SnapThreshold)
+			{
+				RelativeScaleCenter.X = 1.f;
+			}
+		}
+
+		// Snap to center indicator
+		if (FVector2D::Distance(InLocalMousePosition, TransformWidget.BoundsPosition + TransformWidget.BoundsSize * 0.5f) < SnapThreshold)
+		{
+			RelativeScaleCenter = FVector2D(0.5f, 0.5f);
+		}
+
+		DisplayRelativeScaleCenter = RelativeScaleCenter;
 	}
 	else
 	{
@@ -821,49 +930,34 @@ void FCurveEditorTransformTool::OnDrag(const FPointerEvent& InMouseEvent, const
 			const bool bAffectsY = ((TransformWidget.SelectedAnchorFlags & ECurveEditorAnchorFlags::Top) != ECurveEditorAnchorFlags::None ||
 				(TransformWidget.SelectedAnchorFlags & ECurveEditorAnchorFlags::Bottom) != ECurveEditorAnchorFlags::None);
 
-			// We calculate this in [0-1] space for the widget to make the logic easier to follow.
-			FVector2D ScaleCenter = RelativeScaleCenter;
-			const bool bScaleFromEdge = !InMouseEvent.IsAltDown();
-
-			if (bScaleFromEdge)
-			{
-				if ((TransformWidget.SelectedAnchorFlags & ECurveEditorAnchorFlags::Left) != ECurveEditorAnchorFlags::None)
-				{
-					// Anchor to the right side
-					ScaleCenter.X = 1.0f; 
-				}
-
-				if ((TransformWidget.SelectedAnchorFlags & ECurveEditorAnchorFlags::Right) != ECurveEditorAnchorFlags::None)
-				{
-					// Anchor to the left side
-					ScaleCenter.X = 0.0f;
-				}
-
-				if ((TransformWidget.SelectedAnchorFlags & ECurveEditorAnchorFlags::Top) != ECurveEditorAnchorFlags::None)
-				{
-					// Anchor to the bottom side. Slate uses top left as 0 so we flip here.
-					ScaleCenter.Y = 1.0f;
-				}
-
-				if ((TransformWidget.SelectedAnchorFlags & ECurveEditorAnchorFlags::Bottom) != ECurveEditorAnchorFlags::None)
-				{
-					// Anchor to the top side
-					ScaleCenter.Y = 0.0f;
-				}
-			}
-
 			// This is the absolute change since our KeysByCurve was initialized.
 			const FVector2D MouseDelta = InMouseEvent.GetScreenSpacePosition() - InitialMousePosition;
-			// TransformWidget.Size = TransformWidget.StartSize += MouseDelta;
 
 			// We have to flip the delta depending on which edge you grabbed so that the change always goes towards the mouse.
 			const float InputMulSign = (TransformWidget.SelectedAnchorFlags & ECurveEditorAnchorFlags::Left) != ECurveEditorAnchorFlags::None ? -1.0f : 1.0f;
 			const float OutputMulSign = (TransformWidget.SelectedAnchorFlags & ECurveEditorAnchorFlags::Top) != ECurveEditorAnchorFlags::None ? -1.0f : 1.0f;
 			const FVector2D AxisFixedMouseDelta = FVector2D(MouseDelta.X * InputMulSign, MouseDelta.Y * OutputMulSign);
-			const FVector2D PanelSpaceCenter = TransformWidget.StartPosition + (TransformWidget.StartSize * ScaleCenter);
+			const FVector2D PanelSpaceCenter = TransformWidget.StartPosition + (TransformWidget.StartSize * RelativeScaleCenter);
 			const FVector2D ChangeAmount = (AxisFixedMouseDelta / TransformWidget.StartSize);
 
 			ScaleFrom(PanelSpaceCenter, ChangeAmount, bFalloffOn, bAffectsX, bAffectsY);
+
+			TransformWidget.DisplayAnchorFlags = TransformWidget.SelectedAnchorFlags;
+			DisplayRelativeScaleCenter = RelativeScaleCenter;
+			// If x scale flips, toggle left/right flags
+			if (ChangeAmount.X < -1.f)
+			{
+				TransformWidget.DisplayAnchorFlags ^= ECurveEditorAnchorFlags::Left;
+				TransformWidget.DisplayAnchorFlags ^= ECurveEditorAnchorFlags::Right;
+				DisplayRelativeScaleCenter.X = 1 - RelativeScaleCenter.X;
+			}
+			// If y scale flips, toggle top/bottom flags
+			if (ChangeAmount.Y < -1.f)
+			{
+				TransformWidget.DisplayAnchorFlags ^= ECurveEditorAnchorFlags::Top;
+				TransformWidget.DisplayAnchorFlags ^= ECurveEditorAnchorFlags::Bottom;
+				DisplayRelativeScaleCenter.Y = 1 - RelativeScaleCenter.Y;
+			}
 		}
 	}
 
@@ -882,6 +976,8 @@ void FCurveEditorTransformTool::OnDragEnd()
 
 	PrevState.Positon = TransformWidget.Position;
 	PrevState.Size = TransformWidget.Size;
+
+	UpdateToolOptions();
 }
 
 void FCurveEditorTransformTool::StopDragIfPossible()
@@ -989,21 +1085,30 @@ void FCurveEditorTransformTool::Tick(const FGeometry& AllottedGeometry, const do
 		return;
 	}
 
-	if (!FSlateApplication::Get().GetModifierKeys().IsAltDown() && !RelativeScaleCenter.Equals(FVector2D(.5f, .5f)))
+	// Check if all curve models have the same scales
+	TArray<FCurveModelID> CurveModelIDs;
+	CurveEditor->GetSelection().GetAll().GetKeys(CurveModelIDs);
+	bool bHaveSameScales = CurveModelIDs.Num() == 1;
+	if (CurveModelIDs.Num() > 1)
 	{
-		RelativeScaleCenter = FVector2D(.5f, .5f);
+		FCurveEditorScreenSpace FirstCurveViewSpace = CurveEditor->FindFirstInteractiveView(CurveModelIDs[0])->GetCurveSpace(CurveModelIDs[0]);
+		for (int32 ModelIndex = 1; ModelIndex < CurveModelIDs.Num(); ModelIndex++)
+		{
+			FCurveEditorScreenSpace CurveViewSpace = CurveEditor->FindFirstInteractiveView(CurveModelIDs[ModelIndex])->GetCurveSpace(CurveModelIDs[ModelIndex]);
+			bHaveSameScales |= FMath::IsNearlyEqual(CurveViewSpace.PixelsPerInput(), FirstCurveViewSpace.PixelsPerInput()) && FMath::IsNearlyEqual(CurveViewSpace.PixelsPerOutput(), FirstCurveViewSpace.PixelsPerOutput());
+		}
 	}
 
 	// remove tool options widget if less than two selected or there are two different curves
 	TArray<FCurveModelID> CurveModels1;
 	TArray<FCurveModelID> CurveModels2; // no way to reset the curve models between if statements so there has to be two
-	if ((CurveEditor->GetSelection().Count() <= 1 || CurveEditor->GetSelection().GetAll().GetKeys(CurveModels1) != 1) 
+	if ((CurveEditor->GetSelection().Count() <= 1 || !bHaveSameScales) 
 		&& ToolOptionsOnScope != nullptr)
 	{
 		ToolOptionsOnScope = nullptr;
 		OnOptionsRefreshDelegate.Broadcast();
 	}
-	else if ((CurveEditor->GetSelection().Count() > 1 && CurveEditor->GetSelection().GetAll().GetKeys(CurveModels2) == 1) 
+	else if ((CurveEditor->GetSelection().Count() > 1 && bHaveSameScales) 
 		&& ToolOptionsOnScope == nullptr)
 	{
 		ToolOptionsOnScope = MakeShared<FStructOnScope>(FTransformToolOptions::StaticStruct(), (uint8*)&ToolOptions);
diff --git a/Engine/Plugins/Editor/CurveEditorTools/Source/CurveEditorTools/Private/Tools/CurveEditorTransformTool.h b/Engine/Plugins/Editor/CurveEditorTools/Source/CurveEditorTools/Private/Tools/CurveEditorTransformTool.h
--- a/Engine/Plugins/Editor/CurveEditorTools/Source/CurveEditorTools/Private/Tools/CurveEditorTransformTool.h
+++ b/Engine/Plugins/Editor/CurveEditorTools/Source/CurveEditorTools/Private/Tools/CurveEditorTransformTool.h
@@ -28,7 +28,7 @@ enum class ECurveEditorAnchorFlags : uint16
 	FalloffLeft = 0x80,
 	FalloffRight = 0x100,
 	FalloffAll = FalloffTopLeft | FalloffTopRight | FalloffLeft | FalloffRight,
-	CenterScale = 0x200
+	ScaleCenter = 0x200
 };
 
 UENUM()
@@ -55,20 +55,21 @@ struct FCurveEditorTransformWidget
 
 public:
 	ECurveEditorAnchorFlags SelectedAnchorFlags;
+	ECurveEditorAnchorFlags DisplayAnchorFlags;
 
 	FGeometry MakeGeometry(const FGeometry& InWidgetGeometry) const
 	{
 		return InWidgetGeometry.MakeChild(Size, FSlateLayoutTransform(Position));
 	}
 
-	ECurveEditorAnchorFlags GetAnchorFlagsForMousePosition(const FGeometry& InWidgetGeometry, float  FalloffHeight, float FalloffWidth, const FVector2D& RelativeCenterScale, const FVector2D& InMouseScreenPosition) const;
+	ECurveEditorAnchorFlags GetAnchorFlagsForMousePosition(const FGeometry& InWidgetGeometry, float  FalloffHeight, float FalloffWidth, const FVector2D& RelativeScaleCenter, const FVector2D& InMouseScreenPosition) const;
 
 	void GetCenterGeometry(const FGeometry& InWidgetGeometry, FGeometry& OutCenter) const;
 	void GetSidebarGeometry(const FGeometry& InWidgetGeometry, FGeometry& OutLeft, FGeometry& OutRight, FGeometry& OutTop, FGeometry& OutBottom) const;
 	void GetCornerGeometry(const FGeometry& InWidgetGeometry, FGeometry& OutTopLeft, FGeometry& OutTopRight, FGeometry& OutBottomLeft, FGeometry& OutBottomRight) const;
 	void GetFalloffGeometry(const FGeometry& InWidgetGeometry, float FalloffHeight, float FalloffWidth, FGeometry& OutTopLeft, FGeometry& OutTopRight, FGeometry& OutLeft, FGeometry& OutRight) const;
 	void GetScaleCenterGeometry(const FGeometry& InWidgetGeometry, FVector2D ScaleCenter, FGeometry& OutInputBoxGeometry) const;
-
+	void GetCenterIndicatorGeometry(const FGeometry& InWidgetGeometry, FGeometry& OutCenterIndicatorGeometry) const;
 
 	FVector2D Size;
 	FVector2D Position;
@@ -99,6 +100,12 @@ struct FTransformToolOptions
 	UPROPERTY(Transient, EditAnywhere, Category = ToolOptions)
 	float RightBound;
 
+	UPROPERTY(Transient, EditAnywhere, Category = ToolOptions)
+	float ScaleCenterX;
+
+	UPROPERTY(Transient, EditAnywhere, Category = ToolOptions)
+	float ScaleCenterY;
+
 	/** specifies the falloff type applied to curve selection */
 	UPROPERTY(Transient, EditAnywhere, Category = ToolOptions)
 	EToolTransformInterpType FalloffInterpType;
@@ -193,6 +200,9 @@ private:
 	/** point to scale from, relative to current bounds */
 	FVector2D RelativeScaleCenter;
 
+	/** Same as the above, but flipped when the scale amount is inverted for display */
+	FVector2D DisplayRelativeScaleCenter;
+
 	/** Get Soft Select Change Based upon Input Value and internal soft selection weight values*/
 	double GetFalloffWeight(double InputValue) const;
 
diff --git a/Engine/Source/Editor/CurveEditor/Private/Views/SInteractiveCurveEditorView.cpp b/Engine/Source/Editor/CurveEditor/Private/Views/SInteractiveCurveEditorView.cpp
--- a/Engine/Source/Editor/CurveEditor/Private/Views/SInteractiveCurveEditorView.cpp
+++ b/Engine/Source/Editor/CurveEditor/Private/Views/SInteractiveCurveEditorView.cpp
@@ -1063,7 +1063,7 @@ FReply SInteractiveCurveEditorView::OnMouseButtonUp(const FGeometry& MyGeometry,
 	{
 		// Curve Selection Testing.
 		TOptional<FCurveModelID> HitCurve = GetHoveredCurve();
-		if (!HitPoint(MousePosition).IsSet() && HitCurve.IsSet() && !MouseEvent.IsShiftDown() && !MouseEvent.IsAltDown() && !MouseEvent.IsControlDown())
+		if (!HitPoint(MousePosition).IsSet() && HitCurve.IsSet() && !MouseEvent.IsAltDown())
 		{
 			FCurveModel* CurveModel = CurveEditor->FindCurve(HitCurve.GetValue());
 
