Curve Editor: Create new multi scale tool which allows for the scaling of multiple curves from unique pivots

diff --git a/Engine/Plugins/Editor/CurveEditorTools/Source/CurveEditorTools/Private/CurveEditorToolCommands.cpp b/Engine/Plugins/Editor/CurveEditorTools/Source/CurveEditorTools/Private/CurveEditorToolCommands.cpp
--- a/Engine/Plugins/Editor/CurveEditorTools/Source/CurveEditorTools/Private/CurveEditorToolCommands.cpp
+++ b/Engine/Plugins/Editor/CurveEditorTools/Source/CurveEditorTools/Private/CurveEditorToolCommands.cpp
@@ -13,6 +13,7 @@ void FCurveEditorToolCommands::RegisterCommands()
 	// Tool Modes
 	UI_COMMAND(ActivateTransformTool, "Transform Tool", "Activates the Transform tool which allows translation, scale and rotation of selected keys.", EUserInterfaceActionType::ToggleButton, FInputChord(EKeys::W));
 	UI_COMMAND(ActivateRetimeTool, "Retime Tool", "Activates the Retime tool which allows you to define a one dimensional lattice to non-uniformly rescale key times.", EUserInterfaceActionType::ToggleButton, FInputChord(EKeys::E));
+	UI_COMMAND(ActivateMultiScaleTool, "Multi Scale Tool", "Activates the Multi Select tool which allows fast scaling of multiple curves from unique pivots", EUserInterfaceActionType::ToggleButton, FInputChord(EKeys::M));
 }
 
 #undef LOCTEXT_NAMESPACE
diff --git a/Engine/Plugins/Editor/CurveEditorTools/Source/CurveEditorTools/Private/CurveEditorToolCommands.h b/Engine/Plugins/Editor/CurveEditorTools/Source/CurveEditorTools/Private/CurveEditorToolCommands.h
--- a/Engine/Plugins/Editor/CurveEditorTools/Source/CurveEditorTools/Private/CurveEditorToolCommands.h
+++ b/Engine/Plugins/Editor/CurveEditorTools/Source/CurveEditorTools/Private/CurveEditorToolCommands.h
@@ -27,6 +27,7 @@ public:
 	TSharedPtr<FUICommandInfo> SetFocusPlaybackRange;
 	TSharedPtr<FUICommandInfo> ActivateTransformTool;
 	TSharedPtr<FUICommandInfo> ActivateRetimeTool;
+	TSharedPtr<FUICommandInfo> ActivateMultiScaleTool;
 
 public:
 	virtual void RegisterCommands() override;
diff --git a/Engine/Plugins/Editor/CurveEditorTools/Source/CurveEditorTools/Private/CurveEditorToolsModule.cpp b/Engine/Plugins/Editor/CurveEditorTools/Source/CurveEditorTools/Private/CurveEditorToolsModule.cpp
--- a/Engine/Plugins/Editor/CurveEditorTools/Source/CurveEditorTools/Private/CurveEditorToolsModule.cpp
+++ b/Engine/Plugins/Editor/CurveEditorTools/Source/CurveEditorTools/Private/CurveEditorToolsModule.cpp
@@ -6,6 +6,7 @@
 #include "CurveEditorToolCommands.h"
 #include "Tools/CurveEditorTransformTool.h"
 #include "Tools/CurveEditorRetimeTool.h"
+#include "Tools/CurveEditorMultiScaleTool.h"
 #include "Framework/MultiBox/MultiBoxBuilder.h"
 
 #define LOCTEXT_NAMESPACE "CurveEditorToolsModule"
@@ -21,6 +22,7 @@ private:
 	static TSharedRef<ICurveEditorExtension> CreateFocusExtension(TWeakPtr<FCurveEditor> InCurveEditor);
 	static TUniquePtr<ICurveEditorToolExtension> CreateTransformToolExtension(TWeakPtr<FCurveEditor> InCurveEditor);
 	static TUniquePtr<ICurveEditorToolExtension> CreateRetimeToolExtension(TWeakPtr<FCurveEditor> InCurveEditor);
+	static TUniquePtr<ICurveEditorToolExtension> CreateMultiScaleToolExtension(TWeakPtr<FCurveEditor> InCurveEditor);
 
 protected:
 	TSharedRef<FExtender> ExtendCurveEditorToolbarMenu(const TSharedRef<FUICommandList> CommandList);
@@ -30,6 +32,7 @@ private:
 	FDelegateHandle FocusExtensionsHandle;
 	FDelegateHandle TransformToolHandle;
 	FDelegateHandle RetimeToolHandle;
+	FDelegateHandle MultiScaleToolHandle;
 };
 
 IMPLEMENT_MODULE(FCurveEditorToolsModule, CurveEditorTools)
@@ -46,7 +49,7 @@ void FCurveEditorToolsModule::StartupModule()
 	// Register Tool Extensions
 	TransformToolHandle = CurveEditorModule.RegisterToolExtension(FOnCreateCurveEditorToolExtension::CreateStatic(&FCurveEditorToolsModule::CreateTransformToolExtension));
 	RetimeToolHandle = CurveEditorModule.RegisterToolExtension(FOnCreateCurveEditorToolExtension::CreateStatic(&FCurveEditorToolsModule::CreateRetimeToolExtension));
-
+	MultiScaleToolHandle = CurveEditorModule.RegisterToolExtension(FOnCreateCurveEditorToolExtension::CreateStatic(&FCurveEditorToolsModule::CreateMultiScaleToolExtension));
 
 	auto ToolbarExtender = ICurveEditorModule::FCurveEditorMenuExtender::CreateRaw(this, &FCurveEditorToolsModule::ExtendCurveEditorToolbarMenu);
 	auto& MenuExtenders = CurveEditorModule.GetAllToolBarMenuExtenders();
@@ -61,6 +64,7 @@ TSharedRef<FExtender> FCurveEditorToolsModule::ExtendCurveEditorToolbarMenu(cons
 		{
 			ToolbarBuilder.AddToolBarButton(FCurveEditorToolCommands::Get().ActivateTransformTool);
 			ToolbarBuilder.AddToolBarButton(FCurveEditorToolCommands::Get().ActivateRetimeTool);
+			ToolbarBuilder.AddToolBarButton(FCurveEditorToolCommands::Get().ActivateMultiScaleTool);
 		}
 
 		static void FillToolbarFraming(FToolBarBuilder& ToolbarBuilder)
@@ -93,6 +97,7 @@ void FCurveEditorToolsModule::ShutdownModule()
 	// Unregister Tool Extensions
 	CurveEditorModule.UnregisterToolExtension(TransformToolHandle);
 	CurveEditorModule.UnregisterToolExtension(RetimeToolHandle);
+	CurveEditorModule.UnregisterEditorExtension(MultiScaleToolHandle);
 
 	FCurveEditorToolCommands::Unregister();
 }
@@ -112,4 +117,9 @@ TUniquePtr<ICurveEditorToolExtension> FCurveEditorToolsModule::CreateRetimeToolE
 	return MakeUnique<FCurveEditorRetimeTool>(InCurveEditor);
 }
 
+TUniquePtr<ICurveEditorToolExtension> FCurveEditorToolsModule::CreateMultiScaleToolExtension(TWeakPtr<FCurveEditor> InCurveEditor)
+{
+	return MakeUnique<FCurveEditorMultiScaleTool>(InCurveEditor);
+}
+
 #undef LOCTEXT_NAMESPACE
diff --git a/Engine/Plugins/Editor/CurveEditorTools/Source/CurveEditorTools/Private/Tools/CurveEditorMultiScaleTool.cpp b/Engine/Plugins/Editor/CurveEditorTools/Source/CurveEditorTools/Private/Tools/CurveEditorMultiScaleTool.cpp
new file mode 100644
--- /dev/null
+++ b/Engine/Plugins/Editor/CurveEditorTools/Source/CurveEditorTools/Private/Tools/CurveEditorMultiScaleTool.cpp
@@ -0,0 +1,615 @@
+// Copyright 1998-2019 Epic Games, Inc. All Rights Reserved.
+
+#include "CurveEditorMultiScaleTool.h"
+
+#include "CurveEditorToolCommands.h"
+#include "CurveEditor.h"
+#include "Framework/MultiBox/MultiBoxBuilder.h"
+#include "Framework/Commands/UICommandList.h"
+#include "Framework/Commands/UIAction.h"
+#include "ICurveEditorBounds.h"
+#include "Rendering/DrawElements.h"
+#include "CurveEditorScreenSpace.h"
+#include "Editor/Transactor.h"
+#include "ScopedTransaction.h"
+#include "CurveEditorSelection.h"
+#include "CurveModel.h"
+#include "CurveDataAbstraction.h"
+#include "EditorFontGlyphs.h"
+#include "Rendering/SlateRenderer.h"
+#include "Application/SlateApplicationBase.h"
+#include "Fonts/FontMeasure.h"
+#include "CurveEditorSnapMetrics.h"
+#include "SCurveEditorView.h"
+#include "SCurveEditorPanel.h"
+
+#define LOCTEXT_NAMESPACE "CurveEditorToolCommands"
+
+namespace CurveEditorMultiScaleTool
+{
+	constexpr float SliderWidth = 20.f;
+	constexpr float SliderLength = 30.f;
+	constexpr float SliderPadding = 2.f;
+	constexpr float HighlightAlpha = 0.15f;
+}
+
+EMultiScaleAnchorFlags FCurveEditorMultiScaleWidget::GetAnchorFlagsForMousePosition(const FGeometry& InWidgetGeometry, float InXDelta, float InYDelta, const FVector2D& InMouseScreenPosition) const 
+{
+	EMultiScaleAnchorFlags OutFlags = EMultiScaleAnchorFlags::None;
+
+	FGeometry XSliderGeometry, YSliderGeometry, XSidebarGeometry, YSidebarGeometry;
+	GetXSliderGeometry(InWidgetGeometry, InXDelta, XSliderGeometry);
+	GetYSliderGeometry(InWidgetGeometry, InYDelta, YSliderGeometry);
+
+	if (XSliderGeometry.IsUnderLocation(InMouseScreenPosition))
+	{
+		OutFlags |= EMultiScaleAnchorFlags::XSlider;
+	}
+	if (YSliderGeometry.IsUnderLocation(InMouseScreenPosition))
+	{
+		OutFlags |= EMultiScaleAnchorFlags::YSlider;
+	}
+
+	return OutFlags;
+}
+
+void FCurveEditorMultiScaleWidget::GetXSidebarGeometry(const FGeometry& InWidgetGeometry, const FGeometry& InViewGeometry, bool bInSelected, FGeometry& OutSidebar) const
+{
+	const float ViewExtent = InViewGeometry.GetLocalSize().X + 10.f;
+	const FVector2D SidebarSize = FVector2D(bInSelected ? ViewExtent : InWidgetGeometry.GetLocalSize().X, CurveEditorMultiScaleTool::SliderWidth + 2.f * CurveEditorMultiScaleTool::SliderPadding);
+	FVector2D SidebarOffset = FVector2D(bInSelected ? -Position.X - 5.f : (InWidgetGeometry.GetLocalSize().X * 0.5f) - (SidebarSize.X * 0.5f), InWidgetGeometry.GetLocalSize().Y + CurveEditorMultiScaleTool::SliderWidth - CurveEditorMultiScaleTool::SliderPadding);
+
+	OutSidebar = InWidgetGeometry.MakeChild(SidebarSize, FSlateLayoutTransform(SidebarOffset));
+}
+
+void FCurveEditorMultiScaleWidget::GetYSidebarGeometry(const FGeometry& InWidgetGeometry, const FGeometry& InViewGeometry, bool bInSelected, FGeometry& OutSidebar) const
+{
+	const float ViewExtent = InViewGeometry.GetLocalSize().Y + 10.f;
+	const FVector2D SidebarSize = FVector2D(CurveEditorMultiScaleTool::SliderWidth + 2.f * CurveEditorMultiScaleTool::SliderPadding, bInSelected ? ViewExtent : InWidgetGeometry.GetLocalSize().Y);
+	const FVector2D SidebarOffset = FVector2D(InWidgetGeometry.GetLocalSize().X + CurveEditorMultiScaleTool::SliderWidth - CurveEditorMultiScaleTool::SliderPadding, bInSelected ? -Position.Y - 5.f : (InWidgetGeometry.GetLocalSize().Y * 0.5f) - (SidebarSize.Y * 0.5f));
+
+	OutSidebar = InWidgetGeometry.MakeChild(SidebarSize, FSlateLayoutTransform(SidebarOffset));
+}
+
+void FCurveEditorMultiScaleWidget::GetXSliderGeometry(const FGeometry& InWidgetGeometry, float Delta, FGeometry& OutSlider) const
+{
+	const FVector2D SliderSize = FVector2D(CurveEditorMultiScaleTool::SliderLength, CurveEditorMultiScaleTool::SliderWidth);
+	const FVector2D SliderOffset = FVector2D(ParentSpaceDragBegin.X - Position.X + Delta - SliderSize.X * .5f, InWidgetGeometry.GetLocalSize().Y + CurveEditorMultiScaleTool::SliderWidth);
+
+	OutSlider = InWidgetGeometry.MakeChild(SliderSize, FSlateLayoutTransform(SliderOffset));
+}
+
+void FCurveEditorMultiScaleWidget::GetYSliderGeometry(const FGeometry& InWidgetGeometry, float Delta, FGeometry& OutSlider) const
+{
+	const FVector2D SliderSize = FVector2D(CurveEditorMultiScaleTool::SliderWidth, CurveEditorMultiScaleTool::SliderLength);
+	const FVector2D SliderOffset = FVector2D(InWidgetGeometry.GetLocalSize().X + CurveEditorMultiScaleTool::SliderWidth, ParentSpaceDragBegin.Y - Position.Y + Delta - SliderSize.Y * .5f);
+
+	OutSlider = InWidgetGeometry.MakeChild(SliderSize, FSlateLayoutTransform(SliderOffset));
+}
+
+void FCurveEditorMultiScaleTool::OnPaint(const FPaintArgs& Args, const FGeometry& AllottedGeometry, const FSlateRect& MyCullingRect, FSlateWindowElementList& OutDrawElements, int32 PaintOnLayerId, const FWidgetStyle& InWidgetStyle, bool bParentEnabled) const
+{
+	TSharedPtr<FCurveEditor> CurveEditor = WeakCurveEditor.Pin();
+	if (!MultiScaleWidget.Visible || !CurveEditor.IsValid() || (CurveEditor.IsValid() && CurveEditor->GetSelection().Count() <= 1))
+	{
+		return;
+	}
+
+	FGeometry BoundsGeometry = MultiScaleWidget.MakeGeometry(AllottedGeometry);
+
+	const bool bXSliderHovered = (MultiScaleWidget.SelectedAnchorFlags & EMultiScaleAnchorFlags::XSlider) != EMultiScaleAnchorFlags::None;
+	const bool bYSliderHovered = (MultiScaleWidget.SelectedAnchorFlags & EMultiScaleAnchorFlags::YSlider) != EMultiScaleAnchorFlags::None;
+
+	// Draw sliders
+	{
+		FLinearColor XSliderHighlightColor = bXSliderHovered ?
+			FLinearColor::White.CopyWithNewOpacity(CurveEditorMultiScaleTool::HighlightAlpha) : FLinearColor::Transparent;
+		FLinearColor YSliderHighlightColor = bYSliderHovered ?
+			FLinearColor::White.CopyWithNewOpacity(CurveEditorMultiScaleTool::HighlightAlpha) : FLinearColor::Transparent;
+		FGeometry XSliderGeometry, YSliderGeometry;
+		MultiScaleWidget.GetXSliderGeometry(BoundsGeometry, DragDelta.Get(FVector2D::ZeroVector).X, XSliderGeometry);
+		MultiScaleWidget.GetYSliderGeometry(BoundsGeometry, DragDelta.Get(FVector2D::ZeroVector).Y, YSliderGeometry);
+
+		FSlateDrawElement::MakeBox(OutDrawElements, PaintOnLayerId, XSliderGeometry.ToPaintGeometry(), FEditorStyle::GetBrush(TEXT("WhiteBrush")), ESlateDrawEffect::None, XSliderHighlightColor);
+		FSlateDrawElement::MakeBox(OutDrawElements, PaintOnLayerId, XSliderGeometry.ToPaintGeometry(), FEditorStyle::GetBrush(TEXT("MarqueeSelection")));
+
+		FSlateDrawElement::MakeBox(OutDrawElements, PaintOnLayerId, YSliderGeometry.ToPaintGeometry(), FEditorStyle::GetBrush(TEXT("WhiteBrush")), ESlateDrawEffect::None, YSliderHighlightColor);
+		FSlateDrawElement::MakeBox(OutDrawElements, PaintOnLayerId, YSliderGeometry.ToPaintGeometry(), FEditorStyle::GetBrush(TEXT("MarqueeSelection")));
+	}
+
+	// draw Sidebars
+	{
+		FGeometry XSidebarGeometry, YSidebarGeometry;
+		MultiScaleWidget.GetXSidebarGeometry(BoundsGeometry, AllottedGeometry, bXSliderHovered, XSidebarGeometry);
+		MultiScaleWidget.GetYSidebarGeometry(BoundsGeometry, AllottedGeometry, bYSliderHovered, YSidebarGeometry);
+		FSlateDrawElement::MakeBox(OutDrawElements, PaintOnLayerId, XSidebarGeometry.ToPaintGeometry(), FEditorStyle::GetBrush(TEXT("MarqueeSelection")));
+		FSlateDrawElement::MakeBox(OutDrawElements, PaintOnLayerId, YSidebarGeometry.ToPaintGeometry(), FEditorStyle::GetBrush(TEXT("MarqueeSelection")));
+	}
+}
+
+void FCurveEditorMultiScaleTool::Tick(const FGeometry& AllottedGeometry, const double InCurrentTime, const float InDeltaTime)
+{
+	TSharedPtr<FCurveEditor> CurveEditor = WeakCurveEditor.Pin();
+	if (!CurveEditor.IsValid())
+	{
+		return;
+	}
+
+	UpdateBoundingBox();
+
+	FGeometry WidgetGeo = MultiScaleWidget.MakeGeometry(AllottedGeometry);
+
+	if (!DelayedDrag.IsSet())
+	{
+		MultiScaleWidget.ParentSpaceDragBegin = MultiScaleWidget.Position + MultiScaleWidget.Size * 0.5;
+	}
+
+	// Disable options panel if there are no keys selected
+	if (CurveEditor->GetSelection().Count() <= 1 && ToolOptionsOnScope != nullptr)
+	{
+		ToolOptionsOnScope = nullptr;
+		OnOptionsRefreshDelegate.Broadcast();
+	}
+	else if (CurveEditor->GetSelection().Count() > 1 && ToolOptionsOnScope == nullptr)
+	{
+		ToolOptionsOnScope = MakeShared<FStructOnScope>(FMultiScaleToolOptions::StaticStruct(), (uint8*)&ToolOptions);
+		OnOptionsRefreshDelegate.Broadcast();
+	}
+
+	// Keyboard shortcuts to set pivot type
+	if (FSlateApplication::Get().GetModifierKeys().IsLeftShiftDown())
+	{
+		ToolOptions.PivotType = EMultiScalePivotType::FirstKey;
+	}
+	else if (FSlateApplication::Get().GetModifierKeys().IsRightShiftDown())
+	{
+		ToolOptions.PivotType = EMultiScalePivotType::LastKey;
+	}
+}
+
+FReply FCurveEditorMultiScaleTool::OnMouseButtonDown(TSharedRef<SWidget> OwningWidget, const FGeometry& MyGeometry, const FPointerEvent& MouseEvent)
+{
+	TSharedPtr<FCurveEditor> CurveEditor = WeakCurveEditor.Pin();
+	if (!CurveEditor.IsValid())
+	{
+		return FReply::Unhandled();
+	}
+
+	DelayedDrag.Reset();
+	if (CurveEditor->GetSelection().Count() > 1 && MouseEvent.GetEffectingButton() == EKeys::LeftMouseButton)
+	{
+		FGeometry WidgetGeo = MultiScaleWidget.MakeGeometry(MyGeometry);
+		EMultiScaleAnchorFlags HitWidgetFlags = MultiScaleWidget.GetAnchorFlagsForMousePosition(WidgetGeo, 0.f, 0.f, MouseEvent.GetScreenSpacePosition());
+		if (HitWidgetFlags != EMultiScaleAnchorFlags::None)
+		{
+			// Start a Delayed Drag
+			DelayedDrag = FDelayedDrag(MyGeometry.AbsoluteToLocal(MouseEvent.GetScreenSpacePosition()), MouseEvent.GetEffectingButton());
+
+			return FReply::Handled();
+		}
+	}
+	return FReply::Unhandled();
+}
+
+FReply FCurveEditorMultiScaleTool::OnMouseButtonDoubleClick(TSharedRef<SWidget> OwningWidget, const FGeometry& InMyGeometry, const FPointerEvent& InMouseEvent)
+{
+	return FReply::Unhandled();
+}
+
+FReply FCurveEditorMultiScaleTool::OnMouseButtonUp(TSharedRef<SWidget> OwningWidget, const FGeometry& MyGeometry, const FPointerEvent& MouseEvent)
+{
+	FReply Reply = FReply::Handled();
+	if (DelayedDrag.IsSet())
+	{
+		if (DelayedDrag->IsDragging())
+		{
+			OnDragEnd();
+
+			// Only return handled if we actually started a drag
+			Reply.ReleaseMouseCapture();
+		}
+
+		DelayedDrag.Reset();
+		return Reply;
+	}
+
+	return FReply::Unhandled();
+}
+
+FReply FCurveEditorMultiScaleTool::OnMouseMove(TSharedRef<SWidget> OwningWidget, const FGeometry& MyGeometry, const FPointerEvent& MouseEvent)
+{
+	TSharedPtr<FCurveEditor> CurveEditor = WeakCurveEditor.Pin();
+	if (!CurveEditor.IsValid())
+	{
+		return FReply::Unhandled();
+	}
+
+	// Update the Hover State of the widget.
+	if (!DelayedDrag.IsSet())
+	{
+		FGeometry WidgetGeo = MultiScaleWidget.MakeGeometry(MyGeometry);
+
+		EMultiScaleAnchorFlags HitWidgetFlags = MultiScaleWidget.GetAnchorFlagsForMousePosition(WidgetGeo, 0.f, 0.f, MouseEvent.GetScreenSpacePosition());
+		MultiScaleWidget.SelectedAnchorFlags = HitWidgetFlags;
+	}
+
+	const bool bXSliderHovered = (MultiScaleWidget.SelectedAnchorFlags & EMultiScaleAnchorFlags::XSlider) != EMultiScaleAnchorFlags::None;
+	const bool bYSliderHovered = (MultiScaleWidget.SelectedAnchorFlags & EMultiScaleAnchorFlags::YSlider) != EMultiScaleAnchorFlags::None;
+
+	if (DelayedDrag.IsSet())
+	{
+		FReply Reply = FReply::Handled();
+
+		const FVector2D LocalMousePosition = MyGeometry.AbsoluteToLocal(MouseEvent.GetScreenSpacePosition());
+		if (DelayedDrag->IsDragging())
+		{
+			OnDrag(MouseEvent, LocalMousePosition);
+		}
+		else if (DelayedDrag->AttemptDragStart(MouseEvent))
+		{
+			OnDragStart();
+			MultiScaleWidget.ParentSpaceDragBegin.X = bXSliderHovered ? LocalMousePosition.X : MultiScaleWidget.ParentSpaceDragBegin.X;
+			MultiScaleWidget.ParentSpaceDragBegin.Y = bYSliderHovered ? LocalMousePosition.Y : MultiScaleWidget.ParentSpaceDragBegin.Y;
+			// Steal the capture, as we're now the authoritative widget in charge of a mouse-drag operation
+			Reply.CaptureMouse(OwningWidget);
+		}
+
+		return Reply;
+	}
+
+	return FReply::Unhandled();
+}
+
+void FCurveEditorMultiScaleTool::OnFocusLost(const FFocusEvent& InFocusEvent)
+{
+	// We need to end our drag if we lose Window focus to close the transaction, otherwise alt-tabbing while dragging
+	// can cause a transaction to get stuck open.
+	StopDragIfPossible();
+}
+
+void FCurveEditorMultiScaleTool::BindCommands(TSharedRef<FUICommandList> CommandBindings)
+{
+	TSharedPtr<FCurveEditor> CurveEditor = WeakCurveEditor.Pin();
+	if (CurveEditor)
+	{
+		FIsActionChecked MultiScaleToolIsActive = FIsActionChecked::CreateSP(CurveEditor.ToSharedRef(), &FCurveEditor::IsToolActive, ToolID);
+		FExecuteAction ActivateMultiScaleTool = FExecuteAction::CreateSP(CurveEditor.ToSharedRef(), &FCurveEditor::MakeToolActive, ToolID);
+
+		CommandBindings->MapAction(FCurveEditorToolCommands::Get().ActivateMultiScaleTool, ActivateMultiScaleTool, FCanExecuteAction(), MultiScaleToolIsActive);
+	}
+}
+
+void FCurveEditorMultiScaleTool::OnToolOptionsUpdated(const FPropertyChangedEvent& InPropertyChangedEvent)
+{
+	TSharedPtr<FCurveEditor> CurveEditor = WeakCurveEditor.Pin();
+	if (!CurveEditor.IsValid())
+	{
+		return;
+	}
+
+	OnDragStart();
+
+	TArray<FCurveModelID> CurveKeys;
+	CurveEditor->GetSelection().GetAll().GetKeys(CurveKeys);
+	if (CurveKeys.Num() < 1)
+	{
+		return;
+	}
+
+	// grab the view for the first curve, if there are multiple curves it is assumed they all have the same curve space
+	const SCurveEditorView* View = CurveEditor->FindFirstInteractiveView(CurveKeys[0]);
+	if (!View)
+	{
+		return;
+	}
+
+	FCurveEditorScreenSpace CurveSpace = View->GetCurveSpace(CurveKeys[0]);
+
+	const FVector2D ScaleAmount = FVector2D(ToolOptions.XScale - 1.f, ToolOptions.YScale - 1.f);
+	
+	ScaleUnique(ScaleAmount, true, true);
+
+	ToolOptions.XScale = 1.f;
+	ToolOptions.YScale = 1.f;
+	ActiveTransaction.Reset();
+}
+
+void FCurveEditorMultiScaleTool::OnDragStart()
+{
+	ActiveTransaction = MakeUnique<FScopedTransaction>(TEXT("CurveEditorMultiScaleTool"), LOCTEXT("CurveEditorMultiScaleToolTransaction", "Multi Scale Key(s)"), nullptr);
+
+	TSharedPtr<FCurveEditor> CurveEditor = WeakCurveEditor.Pin();
+	if (!CurveEditor.IsValid())
+	{
+		return;
+	}
+
+	// We need to cache our key data because all of our calculations have to be relative to the starting data and not the current per-frame data.
+	KeysByCurve.Reset();
+	bool bMinMaxIsSet = false; // Will be better cached than TOptional and we only need to check this flag here.
+
+	for (const TTuple<FCurveModelID, FKeyHandleSet>& Pair : CurveEditor->GetSelection().GetAll())
+	{
+		FCurveModelID CurveID = Pair.Key;
+		FCurveModel*  Curve = CurveEditor->FindCurve(CurveID);
+
+		if (ensureAlways(Curve))
+		{
+			Curve->Modify();
+
+			TArrayView<const FKeyHandle> Handles = Pair.Value.AsArray();
+
+			FKeyData& KeyData = KeysByCurve.Emplace_GetRef(CurveID);
+			KeyData.Handles = TArray<FKeyHandle>(Handles.GetData(), Handles.Num());
+
+			KeyData.StartKeyPositions.SetNumZeroed(KeyData.Handles.Num());
+			Curve->GetKeyPositions(KeyData.Handles, KeyData.StartKeyPositions);
+			for (const FKeyPosition& KeyPosition : KeyData.StartKeyPositions)
+			{
+				if (bMinMaxIsSet)
+				{
+					if (KeyPosition.InputValue < InputMinMax.GetLowerBoundValue())
+					{
+						InputMinMax.SetLowerBoundValue(KeyPosition.InputValue);
+					}
+					else if (KeyPosition.InputValue > InputMinMax.GetUpperBoundValue())
+					{
+						InputMinMax.SetUpperBoundValue(KeyPosition.InputValue);
+					}
+				}
+				else
+				{
+					bMinMaxIsSet = true;
+					InputMinMax = TRange<double>(KeyPosition.InputValue);
+				}
+			}
+
+			GetPivot(Curve, KeyData);
+		}
+	}
+
+	MultiScaleWidget.StartSize = MultiScaleWidget.Size;
+	MultiScaleWidget.StartPosition = MultiScaleWidget.Position;
+	DragDelta = FVector2D::ZeroVector;
+}
+	
+void FCurveEditorMultiScaleTool::OnDrag(const FPointerEvent& InMouseEvent, const FVector2D& InMousePosition)
+{
+	check(DelayedDrag)
+ 	DragDelta = InMousePosition - DelayedDrag.GetValue().GetInitialPosition();
+	const FVector2D Scale = (DragDelta.GetValue() / (MultiScaleWidget.StartSize * .5f)) * FVector2D(1.f, -1.f);
+
+	const bool bXSliderHovered = (MultiScaleWidget.SelectedAnchorFlags & EMultiScaleAnchorFlags::XSlider) != EMultiScaleAnchorFlags::None;
+	const bool bYSliderHovered = (MultiScaleWidget.SelectedAnchorFlags & EMultiScaleAnchorFlags::YSlider) != EMultiScaleAnchorFlags::None;
+
+	if (bXSliderHovered)
+	{
+		ToolOptions.XScale = 1.f + Scale.X;
+		DragDelta->Y = 0.f;
+		ScaleUnique(FVector2D(Scale.X, 0.f), true, false);
+	}
+	else if (bYSliderHovered)
+	{
+		ToolOptions.YScale = 1.f + Scale.Y;
+		DragDelta->X = 0.f;
+		ScaleUnique(FVector2D(0.f, Scale.Y), false, true);
+	}
+}
+
+void FCurveEditorMultiScaleTool::OnDragEnd()
+{
+	// This finalizes the transaction
+	ActiveTransaction.Reset();
+	DragDelta = TOptional<FVector2D>();
+
+	ToolOptions.XScale = 1.0f;
+	ToolOptions.YScale = 1.0f;
+}
+
+void FCurveEditorMultiScaleTool::StopDragIfPossible()
+{
+	if (DelayedDrag.IsSet())
+	{
+		if (DelayedDrag->IsDragging())
+		{
+			OnDragEnd();
+		}
+
+		DelayedDrag.Reset();
+	}
+}
+
+void FCurveEditorMultiScaleTool::UpdateBoundingBox()
+{
+	TSharedPtr<FCurveEditor> CurveEditor = WeakCurveEditor.Pin();
+	if (!CurveEditor.IsValid())
+	{
+		return;
+	}
+	// We need to look at all selected keys and get their positions relative to the widget view. This lets us put the bounding box around the
+	// current selection, even if it goes off-screen (which it may).
+
+	TOptional<FVector2D> MinValue;
+	TOptional<FVector2D> MaxValue;
+
+	FSlateLayoutTransform AbsoluteToContainer = CurveEditor->GetPanel()->GetViewContainerGeometry().GetAccumulatedLayoutTransform();
+
+	const TMap<FCurveModelID, FKeyHandleSet>& SelectedKeySet = CurveEditor->GetSelection().GetAll();
+	for (const TPair<FCurveModelID, FKeyHandleSet> Pair : SelectedKeySet)
+	{
+		const SCurveEditorView* View = CurveEditor->FindFirstInteractiveView(Pair.Key);
+		if (!View)
+		{
+			continue;
+		}
+
+		// A newly created view may have a zero-size until the next tick which is a problem if
+		// we ask the View for it's curve space, so we skip over it until it has a size.
+		if (View->GetCachedGeometry().GetLocalSize() == FVector2D::ZeroVector)
+		{
+			continue;
+		}
+
+		FCurveModel* CurveModel = CurveEditor->FindCurve(Pair.Key);
+		check(CurveModel);
+
+		TArrayView<const FKeyHandle> KeyHandles = Pair.Value.AsArray();
+
+		TArray<FKeyPosition> KeyPositions;
+		KeyPositions.SetNumUninitialized(KeyHandles.Num());
+		CurveModel->GetKeyPositions(KeyHandles, KeyPositions);
+
+		FCurveEditorScreenSpace CurveSpace = View->GetCurveSpace(Pair.Key);
+		FSlateLayoutTransform InnerToOuterTransform = Concatenate(View->GetCachedGeometry().GetAccumulatedLayoutTransform(), AbsoluteToContainer.Inverse());
+
+		for (int32 i = 0; i < KeyPositions.Num(); i++)
+		{
+			const FVector2D ViewSpaceLocation = FVector2D(CurveSpace.SecondsToScreen(KeyPositions[i].InputValue), CurveSpace.ValueToScreen(KeyPositions[i].OutputValue));
+			const FVector2D PanelSpaceLocation = InnerToOuterTransform.TransformPoint(ViewSpaceLocation);
+
+			if (!MinValue.IsSet())
+			{
+				MinValue = PanelSpaceLocation;
+			}
+
+			if (!MaxValue.IsSet())
+			{
+				MaxValue = PanelSpaceLocation;
+			}
+
+			MinValue = FVector2D::Min(MinValue.GetValue(), PanelSpaceLocation);
+			MaxValue = FVector2D::Max(MaxValue.GetValue(), PanelSpaceLocation);
+		}
+	}
+
+	if (MinValue.IsSet() && MaxValue.IsSet())
+	{
+		FVector2D MarqueeSize = MaxValue.GetValue() - MinValue.GetValue();
+		FVector2D Offset = FVector2D::ZeroVector;
+
+		MultiScaleWidget.BoundsSize = MarqueeSize;
+		MultiScaleWidget.BoundsPosition = MinValue.GetValue();
+
+		// Enforce a minimum size for single time/value selections.
+		if (MarqueeSize.X < 8.f)
+		{
+			MarqueeSize.X = 30;
+			Offset.X = MarqueeSize.X / 2.f;
+		}
+		if (MarqueeSize.Y < 8.f)
+		{
+			MarqueeSize.Y = 30;
+			Offset.Y = MarqueeSize.Y / 2.f;
+		}
+
+		MultiScaleWidget.Visible = true;
+		MultiScaleWidget.Size = MarqueeSize;
+		MultiScaleWidget.Position = MinValue.GetValue() - Offset;
+	}
+	else
+	{
+
+		// No selection, no bounding box.
+		MultiScaleWidget.Visible = false;
+		MultiScaleWidget.Size = FVector2D::ZeroVector;
+		MultiScaleWidget.Position = FVector2D::ZeroVector;
+	}
+}
+
+
+void FCurveEditorMultiScaleTool::ScaleUnique(const FVector2D& InChangeAmount, const bool bInAffectsX, const bool bInAffectsY) const
+{
+	TSharedPtr<FCurveEditor> CurveEditor = WeakCurveEditor.Pin();
+	if (!CurveEditor.IsValid())
+	{
+		return;
+	}
+
+	// Rescaling is where things get tricky. If a user is selecting an edge we scale on one axis at once, and if they select a corner we scale on two axis at once.
+	// If they press alt, we scale relative to the center (instead of relative to the opposite corner). Because these are all bit-flag'd together we can implement
+	// this generically by scaling coordinates relative to an arbitrary center
+
+	const FVector2D PercentChanged = FVector2D(1.0f, 1.0f) + InChangeAmount; // ie: 5 pixel change on a 100 wide gives you 1.05
+
+	FSlateLayoutTransform ContainerToAbsolute = CurveEditor->GetPanel()->GetViewContainerGeometry().GetAccumulatedLayoutTransform().Inverse();
+
+	TArray<FKeyPosition> NewKeyPositionScratch;
+
+	// We now know if we need to affect both X and Y, and we know where we're scaling from. Now we can loop through the keys and actually modify their positions.
+	// We perform the scale on both axis (for simplicity) and then read which axis it should effect before assigning it back to the key position.
+	for (const FKeyData& KeyData : KeysByCurve)
+	{
+		const SCurveEditorView* View = CurveEditor->FindFirstInteractiveView(KeyData.CurveID);
+		if (!View)
+		{
+			continue;
+		}
+
+		FCurveModel* CurveModel = CurveEditor->FindCurve(KeyData.CurveID);
+		check(CurveModel);
+
+		const double CurveSpaceCenterInput = KeyData.Pivot.X;
+		const double CurveSpaceCenterOutput = KeyData.Pivot.Y;
+
+		NewKeyPositionScratch.Reset();
+		NewKeyPositionScratch.Reserve(KeyData.StartKeyPositions.Num());
+
+		FCurveSnapMetrics SnapMetrics = CurveEditor->GetCurveSnapMetrics(KeyData.CurveID);
+
+		for (FKeyPosition StartPosition : KeyData.StartKeyPositions)
+		{
+			// Step 1 is to rescale the position of the key by the percentage change on each axis.
+			const double ScaledInput = (StartPosition.InputValue - CurveSpaceCenterInput) * PercentChanged.X; // *InputMulSign;
+			const double ScaledOutput = (StartPosition.OutputValue - CurveSpaceCenterOutput) * PercentChanged.Y; // *OutputMulSign;
+
+			// Step 2 is to subtract it from the center position so we support scaling from places other than zero
+			const double NewInput = CurveSpaceCenterInput + ScaledInput;
+			const double NewOutput = CurveSpaceCenterOutput + ScaledOutput;
+
+			// Snap the new values to the grid. We calculate both X and Y changes for ease of programming above and just limit which one it applies to.
+			// This includes snapping, otherwise dragging on an edge can cause it to snap on the opposite axis.
+			if (bInAffectsX)
+			{
+				StartPosition.InputValue = View->IsTimeSnapEnabled() ? SnapMetrics.SnapInputSeconds(NewInput) : NewInput;
+			}
+			if (bInAffectsY)
+			{
+				StartPosition.OutputValue = View->IsValueSnapEnabled() ? SnapMetrics.SnapOutput(NewOutput) : NewOutput;
+			}
+
+			NewKeyPositionScratch.Add(StartPosition);
+		}
+
+		CurveModel->SetKeyPositions(KeyData.Handles, NewKeyPositionScratch);
+	}
+}
+
+void FCurveEditorMultiScaleTool::GetPivot(FCurveModel* InCurve, FKeyData& InOutKeyData)
+{
+	switch (ToolOptions.PivotType)
+	{
+	case EMultiScalePivotType::Average:
+		InOutKeyData.Pivot = FVector2D::ZeroVector;
+		for (const FKeyPosition& KeyPosition : InOutKeyData.StartKeyPositions)
+		{
+			InOutKeyData.Pivot += FVector2D(KeyPosition.InputValue, KeyPosition.OutputValue);
+		}
+		InOutKeyData.Pivot /= InOutKeyData.StartKeyPositions.Num();
+		break;
+	case EMultiScalePivotType::BoundCenter:
+		double MinTime, MaxTime, MinValue, MaxValue;
+		InCurve->GetTimeRange(MinTime, MaxTime);
+		InCurve->GetValueRange(MinValue, MaxValue);
+		InOutKeyData.Pivot.X = (MinTime + MaxTime) * .5f;
+		InOutKeyData.Pivot.Y = (MinValue + MaxValue) * .5f;
+		break;
+	case EMultiScalePivotType::FirstKey:
+		InOutKeyData.Pivot.X = InOutKeyData.StartKeyPositions[0].InputValue;
+		InOutKeyData.Pivot.Y = InOutKeyData.StartKeyPositions[0].OutputValue;
+		break;
+	case EMultiScalePivotType::LastKey:
+		InOutKeyData.Pivot.X = InOutKeyData.StartKeyPositions.Top().InputValue;
+		InOutKeyData.Pivot.Y = InOutKeyData.StartKeyPositions.Top().OutputValue;
+		break;
+	}
+}
\ No newline at end of file
diff --git a/Engine/Plugins/Editor/CurveEditorTools/Source/CurveEditorTools/Private/Tools/CurveEditorMultiScaleTool.h b/Engine/Plugins/Editor/CurveEditorTools/Source/CurveEditorTools/Private/Tools/CurveEditorMultiScaleTool.h
new file mode 100644
--- /dev/null
+++ b/Engine/Plugins/Editor/CurveEditorTools/Source/CurveEditorTools/Private/Tools/CurveEditorMultiScaleTool.h
@@ -0,0 +1,172 @@
+// Copyright 1998-2019 Epic Games, Inc. All Rights Reserved.
+
+#pragma once
+#include "CoreMinimal.h"
+#include "ICurveEditorToolExtension.h"
+#include "Layout/Geometry.h"
+#include "Framework/DelayedDrag.h"
+#include "Curves/KeyHandle.h"
+#include "CurveDataAbstraction.h"
+#include "CurveEditorTypes.h"
+#include "CurveEditorSelection.h"
+#include "Misc/EnumClassFlags.h"
+#include "CurveEditorSnapMetrics.h"
+
+#include "CurveEditorMultiScaleTool.generated.h"
+
+class FCurveEditor;
+class FCurveModel;
+struct FCurveModelID;
+struct FKeyHandle;
+
+enum class EMultiScaleAnchorFlags : uint8
+{
+	None = 0x00,
+	XSlider = 0x01,
+	YSlider = 0x02,
+	XSidebar = 0x04,
+	YSidebar = 0x08
+};
+
+ENUM_CLASS_FLAGS(EMultiScaleAnchorFlags);
+
+UENUM()
+enum class EMultiScalePivotType : uint8
+{
+	Average,
+	BoundCenter,
+	FirstKey,
+	LastKey
+};
+
+struct FCurveEditorMultiScaleWidget
+{
+	FCurveEditorMultiScaleWidget()
+	{
+		SelectedAnchorFlags = EMultiScaleAnchorFlags::None;
+		Visible = false;
+	}
+
+public:
+	EMultiScaleAnchorFlags SelectedAnchorFlags;
+
+	FGeometry MakeGeometry(const FGeometry& InWidgetGeometry) const
+	{
+		return InWidgetGeometry.MakeChild(Size, FSlateLayoutTransform(Position));
+	}
+
+	EMultiScaleAnchorFlags GetAnchorFlagsForMousePosition(const FGeometry& InWidgetGeometry, float InUIXScaleFactor, float InUIYScaleFactor, const FVector2D& InMouseScreenPosition) const;
+
+	void GetXSidebarGeometry(const FGeometry& InWidgetGeometry, const FGeometry& InViewGeometry, bool bInSelected, FGeometry& OutSidebar) const;
+	void GetYSidebarGeometry(const FGeometry& InWidgetGeometry, const FGeometry& InViewGeometry, bool bInSelected, FGeometry& OutSidebar) const;
+	
+	void GetXSliderGeometry(const FGeometry& InWidgetGeometry, float InUIScaleFactor, FGeometry& OutSlider) const;
+	void GetYSliderGeometry(const FGeometry& InWidgetGeometry, float InUIScaleFactor, FGeometry& OutSlider) const;
+
+	FVector2D Size;
+	FVector2D Position;
+	bool Visible;
+
+	FVector2D ParentSpaceDragBegin;
+
+	// stores the physical bounds instead of the marquee bounds, for scale calculations
+	FVector2D BoundsSize;
+	FVector2D BoundsPosition;
+
+	FVector2D StartSize;
+	FVector2D StartPosition;
+};
+
+USTRUCT()
+struct FMultiScaleToolOptions
+{
+	GENERATED_BODY()
+
+	UPROPERTY(Transient, EditAnywhere, Category = ToolOptions)
+	float XScale;
+
+	UPROPERTY(Transient, EditAnywhere, Category = ToolOptions)
+	float YScale;
+	
+	UPROPERTY(Transient, EditAnywhere, Category = ToolOptions)
+	EMultiScalePivotType PivotType;
+};
+
+class FCurveEditorMultiScaleTool : public ICurveEditorToolExtension
+{
+public:
+	FCurveEditorMultiScaleTool(TWeakPtr<FCurveEditor> InCurveEditor)
+		: WeakCurveEditor(InCurveEditor)
+	{
+		ToolOptions.XScale = ToolOptions.YScale = 1.f;
+		ToolOptions.PivotType = EMultiScalePivotType::BoundCenter;
+		ToolOptionsOnScope = MakeShared<FStructOnScope>(FMultiScaleToolOptions::StaticStruct(), (uint8*)&ToolOptions);
+	}
+
+	// ICurveEditorToolExtension
+	virtual void OnPaint(const FPaintArgs& Args, const FGeometry& AllottedGeometry, const FSlateRect& MyCullingRect, FSlateWindowElementList& OutDrawElements, int32 PaintOnLayerId, const FWidgetStyle& InWidgetStyle, bool bParentEnabled) const override;
+	virtual void Tick(const FGeometry& AllottedGeometry, const double InCurrentTime, const float InDeltaTime) override;
+	virtual FReply OnMouseButtonDown(TSharedRef<SWidget> OwningWidget, const FGeometry& MyGeometry, const FPointerEvent& MouseEvent) override;
+	virtual FReply OnMouseButtonDoubleClick(TSharedRef<SWidget> OwningWidget, const FGeometry& InMyGeometry, const FPointerEvent& InMouseEvent) override;
+	virtual FReply OnMouseButtonUp(TSharedRef<SWidget> OwningWidget, const FGeometry& MyGeometry, const FPointerEvent& MouseEvent) override;
+	virtual FReply OnMouseMove(TSharedRef<SWidget> OwningWidget, const FGeometry& MyGeometry, const FPointerEvent& MouseEvent) override;
+	virtual void OnFocusLost(const FFocusEvent& InFocusEvent) override;
+	virtual void OnToolActivated() override {}
+	virtual void OnToolDeactivated() override {}
+	virtual void BindCommands(TSharedRef<FUICommandList> CommandBindings) override;
+
+	TSharedPtr<FStructOnScope> GetToolOptions() const override { return ToolOptionsOnScope; }
+	virtual void OnToolOptionsUpdated(const FPropertyChangedEvent& InPropertyChangedEvent) override;
+	// ~ICurveEditorToolExtension
+
+private:
+	void OnDragStart();
+	void OnDrag(const FPointerEvent& InMouseEvent, const FVector2D& InMousePosition);
+	void OnDragEnd();
+	void StopDragIfPossible();
+	void UpdateBoundingBox();
+
+	void ScaleUnique(const FVector2D& InChangeAmount, const bool bInAffectsX, const bool bInAffectsY) const;
+
+private:
+	/** Weak pointer back to the owning curve editor. */
+	TWeakPtr<FCurveEditor> WeakCurveEditor;
+
+	/** The currently open transaction (if any) */
+	TUniquePtr<class FScopedTransaction> ActiveTransaction;
+
+	/** Set when attempting to move a drag handle. */
+	TOptional<FDelayedDrag> DelayedDrag;
+	TOptional<FVector2D> DragDelta;
+
+	/** Used to cache selected key data when doing transform operations */
+	struct FKeyData
+	{
+		FKeyData(FCurveModelID InCurveID)
+			: CurveID(InCurveID)
+		{}
+
+		/** The curve that contains the keys we're dragging */
+		FCurveModelID CurveID;
+		/** All the handles within a given curve that we are dragging */
+		TArray<FKeyHandle> Handles;
+		/** The extended key info for each of the above handles */
+		TArray<FKeyPosition> StartKeyPositions;
+		/** Pivot to scale from */
+		FVector2D Pivot;
+	};
+
+	/** Key dragging data stored per-curve */
+	TArray<FKeyData> KeysByCurve;
+	TRange<double> InputMinMax;
+
+	/** UStruct that displays properties to be modified on screen */
+	FMultiScaleToolOptions ToolOptions;
+	TSharedPtr<FStructOnScope> ToolOptionsOnScope;
+
+	/** Multi select widget to store selection's current state */
+	FCurveEditorMultiScaleWidget MultiScaleWidget;
+
+private:
+	void GetPivot(FCurveModel* InCurve, FKeyData& InOutKeyData);
+};
\ No newline at end of file
