Curve Editor: Add includes to fix errors from last changelist

diff --git a/Engine/Plugins/Editor/CurveEditorTools/Source/CurveEditorTools/Private/Tools/CurveEditorMultiScaleTool.h b/Engine/Plugins/Editor/CurveEditorTools/Source/CurveEditorTools/Private/Tools/CurveEditorMultiScaleTool.h
--- a/Engine/Plugins/Editor/CurveEditorTools/Source/CurveEditorTools/Private/Tools/CurveEditorMultiScaleTool.h
+++ b/Engine/Plugins/Editor/CurveEditorTools/Source/CurveEditorTools/Private/Tools/CurveEditorMultiScaleTool.h
@@ -11,6 +11,8 @@
 #include "CurveEditorSelection.h"
 #include "Misc/EnumClassFlags.h"
 #include "CurveEditorSnapMetrics.h"
+#include "ScopedTransaction.h"
+#include "UObject/StructOnScope.h"
 
 #include "CurveEditorMultiScaleTool.generated.h"
 
