Curve Editor: Fix bug in transform tool where pivot is uninitialized on start

diff --git a/Engine/Plugins/Editor/CurveEditorTools/Source/CurveEditorTools/Private/Tools/CurveEditorTransformTool.h b/Engine/Plugins/Editor/CurveEditorTools/Source/CurveEditorTools/Private/Tools/CurveEditorTransformTool.h
--- a/Engine/Plugins/Editor/CurveEditorTools/Source/CurveEditorTools/Private/Tools/CurveEditorTransformTool.h
+++ b/Engine/Plugins/Editor/CurveEditorTools/Source/CurveEditorTools/Private/Tools/CurveEditorTransformTool.h
@@ -62,7 +62,7 @@ public:
 		return InWidgetGeometry.MakeChild(Size, FSlateLayoutTransform(Position));
 	}
 
-	ECurveEditorAnchorFlags GetAnchorFlagsForMousePosition(const FGeometry& InWidgetGeometry, float  FalloffHeight, float FalloffWidth, const FVector2D& RelativeScaleCenter, const FVector2D& InMouseScreenPosition) const;
+	ECurveEditorAnchorFlags GetAnchorFlagsForMousePosition(const FGeometry& InWidgetGeometry, float FalloffHeight, float FalloffWidth, const FVector2D& RelativeScaleCenter, const FVector2D& InMouseScreenPosition) const;
 
 	void GetCenterGeometry(const FGeometry& InWidgetGeometry, FGeometry& OutCenter) const;
 	void GetSidebarGeometry(const FGeometry& InWidgetGeometry, FGeometry& OutLeft, FGeometry& OutRight, FGeometry& OutTop, FGeometry& OutBottom) const;
@@ -115,10 +115,11 @@ class FCurveEditorTransformTool : public ICurveEditorToolExtension
 {
 public:
 	FCurveEditorTransformTool(TWeakPtr<FCurveEditor> InCurveEditor)
-		: WeakCurveEditor(InCurveEditor), 
-		FalloffHeight(0.0f), 
-		FalloffWidth(0.0f), 
-		RelativeScaleCenter(FVector2D(0.5f, 0.5f))
+		: WeakCurveEditor(InCurveEditor)
+		, FalloffHeight(0.0f)
+		, FalloffWidth(0.0f) 
+		, RelativeScaleCenter(FVector2D(0.5f, 0.5f))
+		, DisplayRelativeScaleCenter(FVector2D(0.5f, 0.5f))
 	{
 		ToolOptions.FalloffInterpType = EToolTransformInterpType::Linear;
 	}
