Curve Editor: Add pivot visualizations to the multi scale tool and fix a bug where the first and last keys for each curve were found incorrectly

diff --git a/Engine/Plugins/Editor/CurveEditorTools/Source/CurveEditorTools/Private/Tools/CurveEditorMultiScaleTool.cpp b/Engine/Plugins/Editor/CurveEditorTools/Source/CurveEditorTools/Private/Tools/CurveEditorMultiScaleTool.cpp
--- a/Engine/Plugins/Editor/CurveEditorTools/Source/CurveEditorTools/Private/Tools/CurveEditorMultiScaleTool.cpp
+++ b/Engine/Plugins/Editor/CurveEditorTools/Source/CurveEditorTools/Private/Tools/CurveEditorMultiScaleTool.cpp
@@ -31,6 +31,7 @@ namespace CurveEditorMultiScaleTool
 	constexpr float SliderLength = 30.f;
 	constexpr float SliderPadding = 2.f;
 	constexpr float HighlightAlpha = 0.15f;
+	constexpr float PivotRadius = 16.f;
 }
 
 EMultiScaleAnchorFlags FCurveEditorMultiScaleWidget::GetAnchorFlagsForMousePosition(const FGeometry& InWidgetGeometry, float InXDelta, float InYDelta, const FVector2D& InMouseScreenPosition) const 
@@ -117,7 +118,7 @@ void FCurveEditorMultiScaleTool::OnPaint(const FPaintArgs& Args, const FGeometry
 		FSlateDrawElement::MakeBox(OutDrawElements, PaintOnLayerId, YSliderGeometry.ToPaintGeometry(), FEditorStyle::GetBrush(TEXT("MarqueeSelection")));
 	}
 
-	// draw Sidebars
+	// Draw Sidebars
 	{
 		FGeometry XSidebarGeometry, YSidebarGeometry;
 		MultiScaleWidget.GetXSidebarGeometry(BoundsGeometry, AllottedGeometry, bXSliderHovered, XSidebarGeometry);
@@ -125,6 +126,61 @@ void FCurveEditorMultiScaleTool::OnPaint(const FPaintArgs& Args, const FGeometry
 		FSlateDrawElement::MakeBox(OutDrawElements, PaintOnLayerId, XSidebarGeometry.ToPaintGeometry(), FEditorStyle::GetBrush(TEXT("MarqueeSelection")));
 		FSlateDrawElement::MakeBox(OutDrawElements, PaintOnLayerId, YSidebarGeometry.ToPaintGeometry(), FEditorStyle::GetBrush(TEXT("MarqueeSelection")));
 	}
+
+	// Draw pivots
+	if (DelayedDrag && DelayedDrag->IsDragging())
+	{
+		check(KeysByCurve.Num() == CurveEditor->GetSelection().GetAll().Num());
+	}
+
+	int32 KeysByCurveIdx = 0;
+	for (TPair<FCurveModelID, FKeyHandleSet> Pair : CurveEditor->GetSelection().GetAll())
+	{
+		FCurveModelID CurveID = DelayedDrag.IsSet() && DelayedDrag->IsDragging() ? KeysByCurve[KeysByCurveIdx].CurveID : Pair.Key;
+		FCurveModel*  Curve = CurveEditor->FindCurve(CurveID);
+
+		if (ensureAlways(Curve))
+		{
+			const SCurveEditorView* View = CurveEditor->FindFirstInteractiveView(CurveID);
+			if (!View)
+			{
+				continue;
+			}
+
+			TArrayView<const FKeyHandle> Handles = Pair.Value.AsArray();
+
+			TArray<FKeyPosition> KeyPositions;
+			KeyPositions.SetNumZeroed(Handles.Num());
+			Curve->GetKeyPositions(Handles, KeyPositions);
+
+			const FVector2D Pivot = DelayedDrag.IsSet() && DelayedDrag->IsDragging() ? KeysByCurve[KeysByCurveIdx].Pivot : GetPivot(Curve, KeyPositions);
+
+			FCurveEditorScreenSpace CurveSpace = View->GetCurveSpace(CurveID);
+			const FVector2D ViewSpacePivot = FVector2D(CurveSpace.SecondsToScreen(Pivot.X), CurveSpace.ValueToScreen(Pivot.Y));
+
+			const FLinearColor PivotColor = Curve->GetColor();
+
+			const FVector2D PivotSize = FVector2D(CurveEditorMultiScaleTool::PivotRadius * 2.f, CurveEditorMultiScaleTool::PivotRadius * 2.f);
+			const FVector2D PivotOffset = ViewSpacePivot - (PivotSize * .5f);
+			FGeometry PivotGeometry = AllottedGeometry.MakeChild(PivotSize, FSlateLayoutTransform(PivotOffset));
+
+			// Draw pivot icon lines
+			const double TangentMul = 2.f;
+			const double Thickness = 1.5f;
+			const FVector2D PivotIconOffset = PivotSize * .5f;
+			FSlateDrawElement::MakeSpline(OutDrawElements, PaintOnLayerId, PivotGeometry.ToPaintGeometry(), FVector2D(-CurveEditorMultiScaleTool::PivotRadius, 0.f) + PivotIconOffset, FVector2D(CurveEditorMultiScaleTool::PivotRadius, 0.f) * TangentMul,
+				FVector2D(0.f, -CurveEditorMultiScaleTool::PivotRadius) + PivotIconOffset, FVector2D(0.f, -CurveEditorMultiScaleTool::PivotRadius) * TangentMul, Thickness, ESlateDrawEffect::None, PivotColor);
+			FSlateDrawElement::MakeSpline(OutDrawElements, PaintOnLayerId, PivotGeometry.ToPaintGeometry(), FVector2D(CurveEditorMultiScaleTool::PivotRadius, 0.f) + PivotIconOffset, FVector2D(-CurveEditorMultiScaleTool::PivotRadius, 0.f) * TangentMul,
+				FVector2D(0.f, CurveEditorMultiScaleTool::PivotRadius) + PivotIconOffset, FVector2D(0.f, CurveEditorMultiScaleTool::PivotRadius) * TangentMul, Thickness, ESlateDrawEffect::None, PivotColor);
+			FSlateDrawElement::MakeSpline(OutDrawElements, PaintOnLayerId, PivotGeometry.ToPaintGeometry(), FVector2D(0.f, CurveEditorMultiScaleTool::PivotRadius) + PivotIconOffset, FVector2D(0.f, -CurveEditorMultiScaleTool::PivotRadius) * TangentMul,
+				FVector2D(-CurveEditorMultiScaleTool::PivotRadius, 0.f) + PivotIconOffset, FVector2D(-CurveEditorMultiScaleTool::PivotRadius, 0.f) * TangentMul, Thickness, ESlateDrawEffect::None, PivotColor);
+			FSlateDrawElement::MakeSpline(OutDrawElements, PaintOnLayerId, PivotGeometry.ToPaintGeometry(), FVector2D(0.f, -CurveEditorMultiScaleTool::PivotRadius) + PivotIconOffset, FVector2D(0.f, CurveEditorMultiScaleTool::PivotRadius) * TangentMul,
+				FVector2D(CurveEditorMultiScaleTool::PivotRadius, 0.f) + PivotIconOffset, FVector2D(CurveEditorMultiScaleTool::PivotRadius, 0.f) * TangentMul, Thickness, ESlateDrawEffect::None, PivotColor);
+		}
+
+		KeysByCurveIdx++;
+	}
+
 }
 
 void FCurveEditorMultiScaleTool::Tick(const FGeometry& AllottedGeometry, const double InCurrentTime, const float InDeltaTime)
@@ -364,7 +420,7 @@ void FCurveEditorMultiScaleTool::OnDragStart()
 				}
 			}
 
-			GetPivot(Curve, KeyData);
+			KeyData.Pivot = GetPivot(Curve, KeyData.StartKeyPositions);
 		}
 	}
 
@@ -584,32 +640,46 @@ void FCurveEditorMultiScaleTool::ScaleUnique(const FVector2D& InChangeAmount, co
 	}
 }
 
-void FCurveEditorMultiScaleTool::GetPivot(FCurveModel* InCurve, FKeyData& InOutKeyData)
+FVector2D FCurveEditorMultiScaleTool::GetPivot(FCurveModel* InCurve, const TArray<FKeyPosition>& InKeyPositions) const
 {
+	const FKeyPosition* FirstKey = nullptr;
+	FVector2D Pivot = FVector2D::ZeroVector;
 	switch (ToolOptions.PivotType)
 	{
 	case EMultiScalePivotType::Average:
-		InOutKeyData.Pivot = FVector2D::ZeroVector;
-		for (const FKeyPosition& KeyPosition : InOutKeyData.StartKeyPositions)
+		check(InKeyPositions.Num() != 0);
+		Pivot = FVector2D::ZeroVector;
+		for (const FKeyPosition& KeyPosition : InKeyPositions)
 		{
-			InOutKeyData.Pivot += FVector2D(KeyPosition.InputValue, KeyPosition.OutputValue);
+			Pivot += FVector2D(KeyPosition.InputValue, KeyPosition.OutputValue);
 		}
-		InOutKeyData.Pivot /= InOutKeyData.StartKeyPositions.Num();
-		break;
+		Pivot /= InKeyPositions.Num();
+		return Pivot;
 	case EMultiScalePivotType::BoundCenter:
 		double MinTime, MaxTime, MinValue, MaxValue;
 		InCurve->GetTimeRange(MinTime, MaxTime);
 		InCurve->GetValueRange(MinValue, MaxValue);
-		InOutKeyData.Pivot.X = (MinTime + MaxTime) * .5f;
-		InOutKeyData.Pivot.Y = (MinValue + MaxValue) * .5f;
-		break;
+		Pivot.X = (MinTime + MaxTime) * .5f;
+		Pivot.Y = (MinValue + MaxValue) * .5f;
+		return Pivot;
 	case EMultiScalePivotType::FirstKey:
-		InOutKeyData.Pivot.X = InOutKeyData.StartKeyPositions[0].InputValue;
-		InOutKeyData.Pivot.Y = InOutKeyData.StartKeyPositions[0].OutputValue;
-		break;
+		FirstKey = Algo::MinElementBy(InKeyPositions, [](const FKeyPosition& Pos) { return Pos.InputValue; });
+		Pivot.X = FirstKey->InputValue;
+		Pivot.Y = FirstKey->OutputValue;
+		return Pivot;
 	case EMultiScalePivotType::LastKey:
-		InOutKeyData.Pivot.X = InOutKeyData.StartKeyPositions.Top().InputValue;
-		InOutKeyData.Pivot.Y = InOutKeyData.StartKeyPositions.Top().OutputValue;
-		break;
+		// Because max element by doesn't exist
+		FKeyPosition Max = FKeyPosition();
+		for (const FKeyPosition& Pos : InKeyPositions)
+		{
+			if (Pos.InputValue > Max.InputValue)
+			{
+				Max = Pos;
+			}
+		}
+		Pivot.X = Max.InputValue;
+		Pivot.Y = Max.OutputValue;
+		return Pivot;
 	}
+	return Pivot;
 }
\ No newline at end of file
diff --git a/Engine/Plugins/Editor/CurveEditorTools/Source/CurveEditorTools/Private/Tools/CurveEditorMultiScaleTool.h b/Engine/Plugins/Editor/CurveEditorTools/Source/CurveEditorTools/Private/Tools/CurveEditorMultiScaleTool.h
--- a/Engine/Plugins/Editor/CurveEditorTools/Source/CurveEditorTools/Private/Tools/CurveEditorMultiScaleTool.h
+++ b/Engine/Plugins/Editor/CurveEditorTools/Source/CurveEditorTools/Private/Tools/CurveEditorMultiScaleTool.h
@@ -170,5 +170,5 @@ private:
 	FCurveEditorMultiScaleWidget MultiScaleWidget;
 
 private:
-	void GetPivot(FCurveModel* InCurve, FKeyData& InOutKeyData);
+	FVector2D GetPivot(FCurveModel* InCurve, const TArray<FKeyPosition>& InKeyPositions) const;
 };
\ No newline at end of file
