Curve Editor: Fix bounds calculations in the multi scale tool

diff --git a/Engine/Plugins/Editor/CurveEditorTools/Source/CurveEditorTools/Private/Tools/CurveEditorMultiScaleTool.cpp b/Engine/Plugins/Editor/CurveEditorTools/Source/CurveEditorTools/Private/Tools/CurveEditorMultiScaleTool.cpp
--- a/Engine/Plugins/Editor/CurveEditorTools/Source/CurveEditorTools/Private/Tools/CurveEditorMultiScaleTool.cpp
+++ b/Engine/Plugins/Editor/CurveEditorTools/Source/CurveEditorTools/Private/Tools/CurveEditorMultiScaleTool.cpp
@@ -22,6 +22,7 @@
 #include "CurveEditorSnapMetrics.h"
 #include "SCurveEditorView.h"
 #include "SCurveEditorPanel.h"
+#include "Algo/MaxElement.h"
 
 #define LOCTEXT_NAMESPACE "CurveEditorToolCommands"
 
@@ -642,8 +643,9 @@ void FCurveEditorMultiScaleTool::ScaleUnique(const FVector2D& InChangeAmount, co
 
 FVector2D FCurveEditorMultiScaleTool::GetPivot(FCurveModel* InCurve, const TArray<FKeyPosition>& InKeyPositions) const
 {
-	const FKeyPosition* FirstKey = nullptr;
 	FVector2D Pivot = FVector2D::ZeroVector;
+	check(InKeyPositions.Num() != 0);
+
 	switch (ToolOptions.PivotType)
 	{
 	case EMultiScalePivotType::Average:
@@ -657,28 +659,24 @@ FVector2D FCurveEditorMultiScaleTool::GetPivot(FCurveModel* InCurve, const TArra
 		return Pivot;
 	case EMultiScalePivotType::BoundCenter:
 		double MinTime, MaxTime, MinValue, MaxValue;
-		InCurve->GetTimeRange(MinTime, MaxTime);
-		InCurve->GetValueRange(MinValue, MaxValue);
+		MinTime = Algo::MinElementBy(InKeyPositions, [](const FKeyPosition& Pos) { return Pos.InputValue; })->InputValue;
+		MaxTime = Algo::MaxElementBy(InKeyPositions, [](const FKeyPosition& Pos) { return Pos.InputValue; })->InputValue;
+		MinValue = Algo::MinElementBy(InKeyPositions, [](const FKeyPosition& Pos) { return Pos.OutputValue; })->OutputValue;
+		MaxValue = Algo::MaxElementBy(InKeyPositions, [](const FKeyPosition& Pos) { return Pos.OutputValue; })->OutputValue;
 		Pivot.X = (MinTime + MaxTime) * .5f;
 		Pivot.Y = (MinValue + MaxValue) * .5f;
 		return Pivot;
 	case EMultiScalePivotType::FirstKey:
+		const FKeyPosition* FirstKey;
 		FirstKey = Algo::MinElementBy(InKeyPositions, [](const FKeyPosition& Pos) { return Pos.InputValue; });
 		Pivot.X = FirstKey->InputValue;
 		Pivot.Y = FirstKey->OutputValue;
 		return Pivot;
 	case EMultiScalePivotType::LastKey:
-		// Because max element by doesn't exist
-		FKeyPosition Max = FKeyPosition();
-		for (const FKeyPosition& Pos : InKeyPositions)
-		{
-			if (Pos.InputValue > Max.InputValue)
-			{
-				Max = Pos;
-			}
-		}
-		Pivot.X = Max.InputValue;
-		Pivot.Y = Max.OutputValue;
+		const FKeyPosition* LastKey;
+		LastKey = Algo::MaxElementBy(InKeyPositions, [](const FKeyPosition& Pos) { return Pos.InputValue; });
+		Pivot.X = LastKey->InputValue;
+		Pivot.Y = LastKey->OutputValue;
 		return Pivot;
 	}
 	return Pivot;
