Curve Editor: Transform tool updates

diff --git a/Engine/Plugins/Editor/CurveEditorTools/Source/CurveEditorTools/Private/Tools/CurveEditorTransformTool.cpp b/Engine/Plugins/Editor/CurveEditorTools/Source/CurveEditorTools/Private/Tools/CurveEditorTransformTool.cpp
--- a/Engine/Plugins/Editor/CurveEditorTools/Source/CurveEditorTools/Private/Tools/CurveEditorTransformTool.cpp
+++ b/Engine/Plugins/Editor/CurveEditorTools/Source/CurveEditorTools/Private/Tools/CurveEditorTransformTool.cpp
@@ -30,6 +30,13 @@ namespace CurveEditorTransformTool
 	constexpr float EdgeHighlightAlpha = 0.15f;
 	constexpr float MaxFalloffOpacity = 0.6f;
 	constexpr int32 FalloffGradientSampleSize = 10;
+
+	constexpr float ArrowMargin = .3f;
+	constexpr float ArrowDiagOffset = 5.f;
+	constexpr float ArrowOffset = 10.f;
+	constexpr float ArrowThreshold = 50.f;
+	constexpr float ArrowExtraRightOffset = 12.f;
+	constexpr float ArrowExtraTopOffset = 12.f;
 }
 
 void FCurveEditorTransformWidget::GetSidebarGeometry(const FGeometry& InWidgetGeometry, FGeometry& OutLeft, FGeometry& OutRight, FGeometry& OutTop, FGeometry& OutBottom) const
@@ -498,8 +505,6 @@ void FCurveEditorTransformTool::OnToolOptionsUpdated(const FPropertyChangedEvent
 
 	UpdateMarqueeBoundingBox();
 
-	PrevState.Positon = TransformWidget.Position;
-	PrevState.Size = TransformWidget.Size;
 	ActiveTransaction.Reset();
 }
 
@@ -550,6 +555,51 @@ void FCurveEditorTransformTool::DrawMarqueeWidget(const FCurveEditorTransformWid
 		FSlateDrawElement::MakeBox(OutDrawElements, InPaintOnLayerId, TopSidebarGeometry.ToPaintGeometry(), FEditorStyle::GetBrush(TEXT("WhiteBrush")), ESlateDrawEffect::None, TopEdgeHighlightColor);
 		// Bottom Edge
 		FSlateDrawElement::MakeBox(OutDrawElements, InPaintOnLayerId, BottomSidebarGeometry.ToPaintGeometry(), FEditorStyle::GetBrush(TEXT("WhiteBrush")), ESlateDrawEffect::None, BottomEdgeHighlightColor);
+
+		// Draw arrow markers if falloff is on
+		if (FSlateApplication::Get().GetModifierKeys().IsControlDown())
+		{
+			// Left/Right sidebar arrows
+			if (LeftSidebarGeometry.GetLocalSize().Y > CurveEditorTransformTool::ArrowThreshold)
+			{
+				{
+					const FVector2D ArrowBottom = FVector2D(-1.f * CurveEditorTransformTool::ArrowOffset, LeftSidebarGeometry.GetLocalSize().Y * (1.f - CurveEditorTransformTool::ArrowMargin));
+					const FVector2D ArrowTop = FVector2D(-1.f * CurveEditorTransformTool::ArrowOffset, LeftSidebarGeometry.GetLocalSize().Y * CurveEditorTransformTool::ArrowMargin);
+					const FVector2D ArrowMiddle = (ArrowBottom + ArrowTop) * .5f;
+					const FVector2D TopArrowDiagOffset = FVector2D(-1.f * CurveEditorTransformTool::ArrowDiagOffset, CurveEditorTransformTool::ArrowDiagOffset);
+					const FVector2D BottomArrowDiagOffset = FVector2D(-1.f * CurveEditorTransformTool::ArrowDiagOffset, -1.f * CurveEditorTransformTool::ArrowDiagOffset);
+					TArray<FVector2D> TopHalfArrow = { ArrowMiddle, ArrowTop, ArrowTop + TopArrowDiagOffset };
+					TArray<FVector2D> BottomHalfArrow = { ArrowMiddle, ArrowBottom, ArrowBottom + BottomArrowDiagOffset };
+					FSlateDrawElement::MakeLines(OutDrawElements, InPaintOnLayerId, LeftSidebarGeometry.ToPaintGeometry(), TopHalfArrow);
+					FSlateDrawElement::MakeLines(OutDrawElements, InPaintOnLayerId, LeftSidebarGeometry.ToPaintGeometry(), BottomHalfArrow);
+				}
+				{
+					const FVector2D ArrowBottom = FVector2D(CurveEditorTransformTool::ArrowOffset + CurveEditorTransformTool::ArrowExtraRightOffset, RightSidebarGeometry.GetLocalSize().Y * (1.f - CurveEditorTransformTool::ArrowMargin));
+					const FVector2D ArrowTop = FVector2D(CurveEditorTransformTool::ArrowOffset + CurveEditorTransformTool::ArrowExtraRightOffset, RightSidebarGeometry.GetLocalSize().Y * CurveEditorTransformTool::ArrowMargin);
+					const FVector2D ArrowMiddle = (ArrowBottom + ArrowTop) * .5f;
+					const FVector2D TopArrowDiagOffset = FVector2D(CurveEditorTransformTool::ArrowDiagOffset, CurveEditorTransformTool::ArrowDiagOffset);
+					const FVector2D BottomArrowDiagOffset = FVector2D(CurveEditorTransformTool::ArrowDiagOffset, -1.f * CurveEditorTransformTool::ArrowDiagOffset);
+					TArray<FVector2D> TopHalfArrow = { ArrowMiddle, ArrowTop, ArrowTop + TopArrowDiagOffset };
+					TArray<FVector2D> BottomHalfArrow = { ArrowMiddle, ArrowBottom, ArrowBottom + BottomArrowDiagOffset };
+					FSlateDrawElement::MakeLines(OutDrawElements, InPaintOnLayerId, RightSidebarGeometry.ToPaintGeometry(), TopHalfArrow);
+					FSlateDrawElement::MakeLines(OutDrawElements, InPaintOnLayerId, RightSidebarGeometry.ToPaintGeometry(), BottomHalfArrow);
+				}
+			}
+
+			// Top sidebar arrows
+			if (TopSidebarGeometry.GetLocalSize().X > CurveEditorTransformTool::ArrowThreshold)
+			{
+				const FVector2D ArrowLeft = FVector2D(TopSidebarGeometry.GetLocalSize().X * CurveEditorTransformTool::ArrowMargin * -1.f, -1.f * CurveEditorTransformTool::ArrowOffset - CurveEditorTransformTool::ArrowExtraTopOffset);
+				const FVector2D ArrowRight = FVector2D(TopSidebarGeometry.GetLocalSize().X * (1.f - CurveEditorTransformTool::ArrowMargin) * -1.f, -1.f * CurveEditorTransformTool::ArrowOffset - CurveEditorTransformTool::ArrowExtraTopOffset);
+				const FVector2D ArrowMiddle = (ArrowLeft + ArrowRight) * .5f;
+				const FVector2D LeftArrowDiagOffset = FVector2D( -1.f * CurveEditorTransformTool::ArrowDiagOffset, -1.f * CurveEditorTransformTool::ArrowDiagOffset);
+				const FVector2D RightArrowDiagOffset = FVector2D(CurveEditorTransformTool::ArrowDiagOffset, -1.f * CurveEditorTransformTool::ArrowDiagOffset);
+				TArray<FVector2D> LeftHalfArrow = { ArrowMiddle, ArrowLeft, ArrowLeft + LeftArrowDiagOffset };
+				TArray<FVector2D> RightHalfArrow = { ArrowMiddle, ArrowRight, ArrowRight + RightArrowDiagOffset };
+				FSlateDrawElement::MakeLines(OutDrawElements, InPaintOnLayerId, RightSidebarGeometry.ToPaintGeometry(), LeftHalfArrow);
+				FSlateDrawElement::MakeLines(OutDrawElements, InPaintOnLayerId, RightSidebarGeometry.ToPaintGeometry(), RightHalfArrow);
+			}
+		}
 	}
 	// Draw falloff if on
 	if (FSlateApplication::Get().GetModifierKeys().IsControlDown())
@@ -763,8 +813,6 @@ void FCurveEditorTransformTool::OnDragStart()
 	TransformWidget.StartSize = TransformWidget.Size;
 	TransformWidget.StartPosition = TransformWidget.Position;
 	SnappingState.Reset();
-	PrevState.Positon = TransformWidget.Position;
-	PrevState.Size = TransformWidget.Size;
 }
 
 void FCurveEditorTransformTool::OnDrag(const FPointerEvent& InMouseEvent, const FVector2D& InLocalMousePosition)
@@ -814,11 +862,11 @@ void FCurveEditorTransformTool::OnDrag(const FPointerEvent& InMouseEvent, const
 
 		const float SnapThreshold = 20.f;
 
+		FVector2D ClosestInRange = InLocalMousePosition;
+		double ClosestInRangeDist = TNumericLimits<double>::Max();
+
 		// Snap to keys
 		{
-			FVector2D ClosestInRange;
-			double ClosestInRangeDist = TNumericLimits<double>::Max();
-
 			for (const TPair<FCurveModelID, FKeyHandleSet>& CurveKeys : CurveEditor->GetSelection().GetAll())
 			{
 				TArray<FKeyPosition> KeyPositions;
@@ -828,7 +876,7 @@ void FCurveEditorTransformTool::OnDrag(const FPointerEvent& InMouseEvent, const
 				const SCurveEditorView* View = CurveEditor->FindFirstInteractiveView(CurveKeys.Key);
 				if (!View)
 				{
-					return;
+					continue;
 				}
 
 				FCurveEditorScreenSpace CurveSpace = View->GetCurveSpace(CurveKeys.Key);
@@ -845,37 +893,83 @@ void FCurveEditorTransformTool::OnDrag(const FPointerEvent& InMouseEvent, const
 				}
 			}
 
-			if (ClosestInRangeDist < SnapThreshold)
+			// Not snapped, reset everything
+			if (ClosestInRangeDist > SnapThreshold)
+			{
+				ClosestInRange = InLocalMousePosition;
+				ClosestInRangeDist = TNumericLimits<double>::Max();
+			}
+		}
+
+		// Snap to grid
+		if (bCurvesHaveSameScales && (CurveEditor->IsInputSnappingEnabled() || CurveEditor->IsOutputSnappingEnabled()))
+		{
+			// Snap the mouse to the grid according to the first curve's snap metrics (assuming all curves have the same view scales)
+			const FCurveModelID FirstCurveID = (*CurveEditor->GetSelection().GetAll().CreateConstIterator()).Key;
+
+			const SCurveEditorView* View = CurveEditor->FindFirstInteractiveView(FirstCurveID);
+			if (View)
 			{
-				const FVector2D ViewSpaceDelta = ClosestInRange - TransformWidget.BoundsPosition;
-				RelativeScaleCenter = ViewSpaceDelta / TransformWidget.BoundsSize;
+				FCurveEditorScreenSpace CurveSpace = View->GetCurveSpace(FirstCurveID);
+
+				FCurveSnapMetrics SnapMetrics = CurveEditor->GetCurveSnapMetrics(FirstCurveID);
+				const FVector2D CurveSpaceSnappedMousePoint = FVector2D(SnapMetrics.SnapInputSeconds(CurveSpace.ScreenToSeconds(InLocalMousePosition.X)), SnapMetrics.SnapOutput(CurveSpace.ScreenToValue(InLocalMousePosition.Y)));
+				const FVector2D ViewSpaceSnappedMousePoint = FVector2D(CurveSpace.SecondsToScreen(CurveSpaceSnappedMousePoint.X), CurveSpace.ValueToScreen(CurveSpaceSnappedMousePoint.Y));
+
+				float Dist = 0.f;
+				if ((Dist = FVector2D::Distance(ViewSpaceSnappedMousePoint, InLocalMousePosition)) < ClosestInRangeDist)
+				{
+					ClosestInRange = ViewSpaceSnappedMousePoint;
+					ClosestInRangeDist = Dist;
+				}
 			}
 		}
 
 		// Snap to edges or corners
 		{
-			if (FMath::Abs(InLocalMousePosition.Y - (TransformWidget.Position.Y)) < SnapThreshold)
+			float Dist = FMath::Abs(InLocalMousePosition.Y - (TransformWidget.Position.Y));
+			if (Dist < FMath::Min<float>(SnapThreshold, ClosestInRangeDist))
 			{
-				RelativeScaleCenter.Y = 0.f;
+				ClosestInRange.Y = TransformWidget.Position.Y;
+				ClosestInRangeDist = Dist;
 			}
-			if (FMath::Abs(InLocalMousePosition.Y - (TransformWidget.Position.Y + TransformWidget.Size.Y)) < SnapThreshold)
+
+			Dist = FMath::Abs(InLocalMousePosition.Y - (TransformWidget.Position.Y + TransformWidget.Size.Y));
+			if (Dist < FMath::Min<float>(SnapThreshold, ClosestInRangeDist))
 			{
-				RelativeScaleCenter.Y = 1.f;
+				ClosestInRange.Y = TransformWidget.Position.Y + TransformWidget.Size.Y;
+				ClosestInRangeDist = Dist;
 			}
-			if (FMath::Abs(InLocalMousePosition.X - (TransformWidget.Position.X)) < SnapThreshold)
+
+			Dist = FMath::Abs(InLocalMousePosition.X - (TransformWidget.Position.X));
+			if (Dist < FMath::Min<float>(SnapThreshold, ClosestInRangeDist))
 			{
-				RelativeScaleCenter.X = 0.f;
+				ClosestInRange.X = TransformWidget.Position.X;
+				ClosestInRangeDist = Dist;
 			}
-			if (FMath::Abs(InLocalMousePosition.X - (TransformWidget.Position.X + TransformWidget.Size.X)) < SnapThreshold)
+
+			Dist = FMath::Abs(InLocalMousePosition.X - (TransformWidget.Position.X + TransformWidget.Size.X));
+			if (Dist < FMath::Min<float>(SnapThreshold, ClosestInRangeDist))
 			{
-				RelativeScaleCenter.X = 1.f;
+				ClosestInRange.X = TransformWidget.Position.X + TransformWidget.Size.X;
+				ClosestInRangeDist = Dist;
 			}
 		}
 
 		// Snap to center indicator
-		if (FVector2D::Distance(InLocalMousePosition, TransformWidget.BoundsPosition + TransformWidget.BoundsSize * 0.5f) < SnapThreshold)
 		{
-			RelativeScaleCenter = FVector2D(0.5f, 0.5f);
+			float Dist = 0.f;
+			if ((Dist = FVector2D::Distance(InLocalMousePosition, TransformWidget.BoundsPosition + TransformWidget.BoundsSize * 0.5f)) < ClosestInRangeDist)
+			{
+				ClosestInRange = TransformWidget.BoundsPosition + TransformWidget.BoundsSize * 0.5f;
+				ClosestInRangeDist = Dist;
+			}
+		}
+
+		if (ClosestInRangeDist < SnapThreshold)
+		{
+			const FVector2D ViewSpaceDelta = ClosestInRange - TransformWidget.BoundsPosition;
+			RelativeScaleCenter = ViewSpaceDelta / TransformWidget.BoundsSize;
 		}
 
 		DisplayRelativeScaleCenter = RelativeScaleCenter;
@@ -966,9 +1060,6 @@ void FCurveEditorTransformTool::OnDrag(const FPointerEvent& InMouseEvent, const
 	UpdateMarqueeBoundingBox();
 
 	UpdateToolOptions();
-
-	PrevState.Positon = TransformWidget.Position;
-	PrevState.Size = TransformWidget.Size;
 }
 
 void FCurveEditorTransformTool::OnDragEnd()
@@ -984,9 +1075,6 @@ void FCurveEditorTransformTool::OnDragEnd()
 	// This finalizes the transaction
 	ActiveTransaction.Reset();
 
-	PrevState.Positon = TransformWidget.Position;
-	PrevState.Size = TransformWidget.Size;
-
 	UpdateToolOptions();
 }
 
@@ -1098,40 +1186,50 @@ void FCurveEditorTransformTool::Tick(const FGeometry& AllottedGeometry, const do
 	// Check if all curve models have the same scales
 	TArray<FCurveModelID> CurveModelIDs;
 	CurveEditor->GetSelection().GetAll().GetKeys(CurveModelIDs);
-	bool bHaveSameScales = CurveModelIDs.Num() == 1;
+	bCurvesHaveSameScales = CurveModelIDs.Num() == 1;
 	if (CurveModelIDs.Num() > 1)
 	{
 		FCurveEditorScreenSpace FirstCurveViewSpace = CurveEditor->FindFirstInteractiveView(CurveModelIDs[0])->GetCurveSpace(CurveModelIDs[0]);
 		for (int32 ModelIndex = 1; ModelIndex < CurveModelIDs.Num(); ModelIndex++)
 		{
 			FCurveEditorScreenSpace CurveViewSpace = CurveEditor->FindFirstInteractiveView(CurveModelIDs[ModelIndex])->GetCurveSpace(CurveModelIDs[ModelIndex]);
-			bHaveSameScales |= FMath::IsNearlyEqual(CurveViewSpace.PixelsPerInput(), FirstCurveViewSpace.PixelsPerInput()) && FMath::IsNearlyEqual(CurveViewSpace.PixelsPerOutput(), FirstCurveViewSpace.PixelsPerOutput());
+			bCurvesHaveSameScales |= FMath::IsNearlyEqual(CurveViewSpace.PixelsPerInput(), FirstCurveViewSpace.PixelsPerInput()) && FMath::IsNearlyEqual(CurveViewSpace.PixelsPerOutput(), FirstCurveViewSpace.PixelsPerOutput());
 		}
 	}
 
 	// remove tool options widget if less than two selected or there are two different curves
-	TArray<FCurveModelID> CurveModels1;
-	TArray<FCurveModelID> CurveModels2; // no way to reset the curve models between if statements so there has to be two
-	if ((CurveEditor->GetSelection().Count() <= 1 || !bHaveSameScales) 
+	if ((CurveEditor->GetSelection().Count() <= 1 || !bCurvesHaveSameScales) 
 		&& ToolOptionsOnScope != nullptr)
 	{
 		ToolOptionsOnScope = nullptr;
 		OnOptionsRefreshDelegate.Broadcast();
 	}
-	else if ((CurveEditor->GetSelection().Count() > 1 && bHaveSameScales) 
+	else if ((CurveEditor->GetSelection().Count() > 1 && bCurvesHaveSameScales) 
 		&& ToolOptionsOnScope == nullptr)
 	{
 		ToolOptionsOnScope = MakeShared<FStructOnScope>(FTransformToolOptions::StaticStruct(), (uint8*)&ToolOptions);
 		OnOptionsRefreshDelegate.Broadcast();
 	}
 
-	// check if there have been any changes made outside of our tool
-	if (PrevState.Positon != TransformWidget.BoundsPosition || PrevState.Size != TransformWidget.BoundsSize)
+	// check if there have been any changes made outside of the tool
+	if (bCurvesHaveSameScales)
 	{
-		UpdateToolOptions();
+		const FCurveModelID FirstCurveID = (*CurveEditor->GetSelection().GetAll().CreateConstIterator()).Key;
 
-		PrevState.Positon = TransformWidget.Position;
-		PrevState.Size = TransformWidget.Size;
+		const SCurveEditorView* View = CurveEditor->FindFirstInteractiveView(FirstCurveID);
+		if (!View)
+		{
+			return;
+		}
+
+		FCurveEditorScreenSpace CurveSpace = View->GetCurveSpace(FirstCurveID);
+
+		const FVector2D OptionsSize = FVector2D(ToolOptions.RightBound - ToolOptions.LeftBound, ToolOptions.LowerBound - ToolOptions.UpperBound);
+		const FVector2D ViewSpaceOptionsSize = FVector2D(CurveSpace.PixelsPerInput() * FMath::Abs(OptionsSize.X), CurveSpace.PixelsPerOutput() * FMath::Abs(OptionsSize.Y));
+		if ((!DelayedDrag || (DelayedDrag && !DelayedDrag->IsDragging())) && (TransformWidget.BoundsSize - ViewSpaceOptionsSize).IsNearlyZero())
+		{
+			UpdateToolOptions();
+		}
 	}
 }
 
diff --git a/Engine/Plugins/Editor/CurveEditorTools/Source/CurveEditorTools/Private/Tools/CurveEditorTransformTool.h b/Engine/Plugins/Editor/CurveEditorTools/Source/CurveEditorTools/Private/Tools/CurveEditorTransformTool.h
--- a/Engine/Plugins/Editor/CurveEditorTools/Source/CurveEditorTools/Private/Tools/CurveEditorTransformTool.h
+++ b/Engine/Plugins/Editor/CurveEditorTools/Source/CurveEditorTools/Private/Tools/CurveEditorTransformTool.h
@@ -107,7 +107,7 @@ struct FTransformToolOptions
 	float ScaleCenterY;
 
 	/** specifies the falloff type applied to curve selection */
-	UPROPERTY(Transient, EditAnywhere, Category = ToolOptions)
+	UPROPERTY(Transient, EditAnywhere, Category = ToolOptions, Meta = (ToolTip = "Interpolation type for soft selection (activate by holding ctrl)"))
 	EToolTransformInterpType FalloffInterpType;
 };
 
@@ -188,16 +188,15 @@ private:
 	FTransformToolOptions ToolOptions;
 	TSharedPtr<FStructOnScope> ToolOptionsOnScope;
 
-	/** helper struct for checking if tool has changed outside of own operations, i.e. middle click */
-	struct ToolBounds { FVector2D Positon, Size; };
-	ToolBounds PrevState;
-
 private:
 	/** Soft Selection Height, if 0.0 as we go up to 1.0 we go up to the top edge */
 	float FalloffHeight;
 	/** Soft Selection Width, if 0.0 we peak directly at the midpoint gradually scaling out till we reach 1.0fff */
 	float FalloffWidth;
 
+	/** Whether or not the selections have the same view transforms */
+	bool bCurvesHaveSameScales;
+
 	/** point to scale from, relative to current bounds */
 	FVector2D RelativeScaleCenter;
 
diff --git a/Engine/Source/Editor/CurveEditor/Private/SCurveEditorViewContainer.cpp b/Engine/Source/Editor/CurveEditor/Private/SCurveEditorViewContainer.cpp
--- a/Engine/Source/Editor/CurveEditor/Private/SCurveEditorViewContainer.cpp
+++ b/Engine/Source/Editor/CurveEditor/Private/SCurveEditorViewContainer.cpp
@@ -131,6 +131,8 @@ FReply SCurveEditorViewContainer::OnKeyDown(const FGeometry& MyGeometry, const F
 
 FReply SCurveEditorViewContainer::OnPreviewMouseButtonDown(const FGeometry& MyGeometry, const FPointerEvent& MouseEvent)
 {
+	bCaughtMouseDown = true;
+
 	if (CurveEditor->GetCurrentTool())
 	{
 		return CurveEditor->GetCurrentTool()->OnMouseButtonDown(AsShared(), MyGeometry, MouseEvent);
@@ -237,7 +239,7 @@ FReply SCurveEditorViewContainer::OnMouseButtonUp(const FGeometry& MyGeometry, c
 		Reply = CurveEditor->GetCurrentTool()->OnMouseButtonUp(AsShared(), MyGeometry, MouseEvent);
 	}
 
-	if (!Reply.IsEventHandled() && MouseEvent.GetEffectingButton() == EKeys::LeftMouseButton)
+	if (!Reply.IsEventHandled() && MouseEvent.GetEffectingButton() == EKeys::LeftMouseButton && bCaughtMouseDown)
 	{
 		if (!MouseEvent.IsShiftDown() && !MouseEvent.IsAltDown() && !MouseEvent.IsControlDown())
 		{
@@ -246,6 +248,7 @@ FReply SCurveEditorViewContainer::OnMouseButtonUp(const FGeometry& MyGeometry, c
 		}
 	}
 
+	bCaughtMouseDown = false;
 	DragOperation.Reset();
 	return Reply.ReleaseMouseCapture();
 }
diff --git a/Engine/Source/Editor/CurveEditor/Private/SCurveEditorViewContainer.h b/Engine/Source/Editor/CurveEditor/Private/SCurveEditorViewContainer.h
--- a/Engine/Source/Editor/CurveEditor/Private/SCurveEditorViewContainer.h
+++ b/Engine/Source/Editor/CurveEditor/Private/SCurveEditorViewContainer.h
@@ -79,4 +79,10 @@ private:
 
 	/** Array of views that may need their height updating on tick. */
 	TArray<TSharedPtr<SCurveEditorView>> Views;
+
+	/** 
+	 * Whether or not this widget caught an OnMouseDown notification 
+	 * Used to check if the selection should be cleared
+	 */
+	bool bCaughtMouseDown;
 };
\ No newline at end of file
