Curve Editor: FFT would crash due to too small baking interval.

diff --git a/Engine/Plugins/Editor/CurveEditorTools/Source/CurveEditorTools/Private/Filters/CurveEditorFFTFilter.cpp b/Engine/Plugins/Editor/CurveEditorTools/Source/CurveEditorTools/Private/Filters/CurveEditorFFTFilter.cpp
--- a/Engine/Plugins/Editor/CurveEditorTools/Source/CurveEditorTools/Private/Filters/CurveEditorFFTFilter.cpp
+++ b/Engine/Plugins/Editor/CurveEditorTools/Source/CurveEditorTools/Private/Filters/CurveEditorFFTFilter.cpp
@@ -25,7 +25,6 @@ void UCurveEditorFFTFilter::ApplyFilter_Impl(TSharedRef<FCurveEditor> InCurveEdi
 	// Since we're baking under the hood, we need to cache their bake interval, override it to a reasonable number, and restore it at the end.
 	float OriginalIntervalRate = BakeFilter->BakeInterval;
 	bool OriginalUseSnapRate = BakeFilter->bUseSnapRateForInterval;
-	BakeFilter->BakeInterval = 0.1f;
 	BakeFilter->bUseSnapRateForInterval = false;
 
 	TArray<FKeyHandle> OriginalKeyHandles;
@@ -55,12 +54,13 @@ void UCurveEditorFFTFilter::ApplyFilter_Impl(TSharedRef<FCurveEditor> InCurveEdi
 		OriginalKeyHandles.Reset();
 		CurveModel->GetKeys(*InCurveEditor, MinKey, MaxKey, TNumericLimits<double>::Lowest(), TNumericLimits<double>::Max(), OriginalKeyHandles);
 		
-		// Need at least two keys to filter
-		if (OriginalKeyHandles.Num() < 2)
+		// Need at least three keys to filter
+		if (OriginalKeyHandles.Num() < 3)
 		{
 			continue;
 		}
-
+		//Set interval range divded by twice the number of keys.
+		BakeFilter->BakeInterval = (MaxKey - MinKey) /  ( 2 * OriginalKeyHandles.Num());
 		// This stores the position of all of the original keys. Once we've filtered the curve we will need to sample it at these positions.
 		TArray<FKeyPosition> OriginalKeyPositions;
 		TArray<FKeyAttributes> OriginalKeyAttributes;
