Curve Editor: Change mode hotkeys to have a modifier. Transform - Ctrl-T, Retime - Ctrl-E, Multi - Ctrl-M. This removes conflicts with w,e,r which are reserved for the viewport transform gizmo. Ctrl-W is used for weighted tangents.

diff --git a/Engine/Plugins/Editor/CurveEditorTools/Source/CurveEditorTools/Private/CurveEditorToolCommands.cpp b/Engine/Plugins/Editor/CurveEditorTools/Source/CurveEditorTools/Private/CurveEditorToolCommands.cpp
--- a/Engine/Plugins/Editor/CurveEditorTools/Source/CurveEditorTools/Private/CurveEditorToolCommands.cpp
+++ b/Engine/Plugins/Editor/CurveEditorTools/Source/CurveEditorTools/Private/CurveEditorToolCommands.cpp
@@ -11,9 +11,9 @@ void FCurveEditorToolCommands::RegisterCommands()
 	UI_COMMAND(SetFocusPlaybackRange, "Focus Playback Range", "Focuses the Curve Editor on the current Playback Range with zoom based on visible curves.", EUserInterfaceActionType::RadioButton, FInputChord(EKeys::A));
 
 	// Tool Modes
-	UI_COMMAND(ActivateTransformTool, "Transform Tool", "Activates the Transform tool which allows translation, scale and rotation of selected keys.", EUserInterfaceActionType::ToggleButton, FInputChord(EKeys::W));
-	UI_COMMAND(ActivateRetimeTool, "Retime Tool", "Activates the Retime tool which allows you to define a one dimensional lattice to non-uniformly rescale key times.", EUserInterfaceActionType::ToggleButton, FInputChord(EKeys::E));
-	UI_COMMAND(ActivateMultiScaleTool, "Multi Scale Tool", "Activates the Multi Select tool which allows fast scaling of multiple curves from unique pivots", EUserInterfaceActionType::ToggleButton, FInputChord(EKeys::M));
+	UI_COMMAND(ActivateTransformTool, "Transform Tool", "Activates the Transform tool which allows translation, scale and rotation of selected keys.", EUserInterfaceActionType::ToggleButton, FInputChord(EModifierKey::Control, EKeys::T));
+	UI_COMMAND(ActivateRetimeTool, "Retime Tool", "Activates the Retime tool which allows you to define a one dimensional lattice to non-uniformly rescale key times.", EUserInterfaceActionType::ToggleButton, FInputChord(EModifierKey::Control, EKeys::E));
+	UI_COMMAND(ActivateMultiScaleTool, "Multi Scale Tool", "Activates the Multi Select tool which allows fast scaling of multiple curves from unique pivots", EUserInterfaceActionType::ToggleButton, FInputChord(EModifierKey::Control, EKeys::M));
 }
 
 #undef LOCTEXT_NAMESPACE
