CurveEditor: Scaling with Transform Tool would revert on mouse up.

diff --git a/Engine/Plugins/Editor/CurveEditorTools/Source/CurveEditorTools/Private/Tools/CurveEditorTransformTool.cpp b/Engine/Plugins/Editor/CurveEditorTools/Source/CurveEditorTools/Private/Tools/CurveEditorTransformTool.cpp
--- a/Engine/Plugins/Editor/CurveEditorTools/Source/CurveEditorTools/Private/Tools/CurveEditorTransformTool.cpp
+++ b/Engine/Plugins/Editor/CurveEditorTools/Source/CurveEditorTools/Private/Tools/CurveEditorTransformTool.cpp
@@ -1101,7 +1101,7 @@ void FCurveEditorTransformTool::StopDragIfPossible()
 	}
 }
 
-void FCurveEditorTransformTool::ScaleFrom(const FVector2D& InPanelSpaceCenter, const FVector2D& InChangeAmount, const bool bInFalloffOn, const bool bInAffectsX, const bool bInAffectsY) const
+void FCurveEditorTransformTool::ScaleFrom(const FVector2D& InPanelSpaceCenter, const FVector2D& InChangeAmount, const bool bInFalloffOn, const bool bInAffectsX, const bool bInAffectsY) 
 {
 	TSharedPtr<FCurveEditor> CurveEditor = WeakCurveEditor.Pin();
 	if (!CurveEditor)
@@ -1121,7 +1121,7 @@ void FCurveEditorTransformTool::ScaleFrom(const FVector2D& InPanelSpaceCenter, c
 
 	// We now know if we need to affect both X and Y, and we know where we're scaling from. Now we can loop through the keys and actually modify their positions.
 	// We perform the scale on both axis (for simplicity) and then read which axis it should effect before assigning it back to the key position.
-	for (const FKeyData& KeyData : KeysByCurve)
+	for (FKeyData& KeyData : KeysByCurve)
 	{
 		const SCurveEditorView* View = CurveEditor->FindFirstInteractiveView(KeyData.CurveID);
 		if (!View)
@@ -1177,6 +1177,8 @@ void FCurveEditorTransformTool::ScaleFrom(const FVector2D& InPanelSpaceCenter, c
 		}
 
 		CurveModel->SetKeyPositions(KeyData.Handles, NewKeyPositionScratch, EPropertyChangeType::Interactive);
+		KeyData.LastDraggedKeyPositions = NewKeyPositionScratch;
+
 	}
 }
 
diff --git a/Engine/Plugins/Editor/CurveEditorTools/Source/CurveEditorTools/Private/Tools/CurveEditorTransformTool.h b/Engine/Plugins/Editor/CurveEditorTools/Source/CurveEditorTools/Private/Tools/CurveEditorTransformTool.h
--- a/Engine/Plugins/Editor/CurveEditorTools/Source/CurveEditorTools/Private/Tools/CurveEditorTransformTool.h
+++ b/Engine/Plugins/Editor/CurveEditorTools/Source/CurveEditorTools/Private/Tools/CurveEditorTransformTool.h
@@ -142,7 +142,7 @@ private:
 	void UpdateMarqueeBoundingBox();
 	void UpdateToolOptions();
 	void DrawMarqueeWidget(const FCurveEditorTransformWidget& InTransformWidget, const FPaintArgs& InArgs, const FGeometry& InAllottedGeometry, const FSlateRect& InMyCullingRect, FSlateWindowElementList& OutDrawElements, const int32 InPaintOnLayerId, const FWidgetStyle& InWidgetStyle, const bool bInParentEnabled) const;
-	void ScaleFrom(const FVector2D& InPanelSpaceCenter, const FVector2D& InChangeAmount, const bool bInFalloffOn, const bool bInAffectsX, const bool bInAffectsY) const;
+	void ScaleFrom(const FVector2D& InPanelSpaceCenter, const FVector2D& InChangeAmount, const bool bInFalloffOn, const bool bInAffectsX, const bool bInAffectsY);
 
 	void OnDragStart();
 	void OnDrag(const FPointerEvent& InMouseEvent, const FVector2D& InLocalMousePosition);
