Curve Editor: Remove any overlapping keys on mouse up (value set).

Gather up all the final drag times of the keys. For each time, get all the keys that match that time and keep the key that moved the most.

diff --git a/Engine/Plugins/Editor/CurveEditorTools/Source/CurveEditorTools/Private/Tools/CurveEditorTransformTool.cpp b/Engine/Plugins/Editor/CurveEditorTools/Source/CurveEditorTools/Private/Tools/CurveEditorTransformTool.cpp
--- a/Engine/Plugins/Editor/CurveEditorTools/Source/CurveEditorTools/Private/Tools/CurveEditorTransformTool.cpp
+++ b/Engine/Plugins/Editor/CurveEditorTools/Source/CurveEditorTools/Private/Tools/CurveEditorTransformTool.cpp
@@ -796,12 +796,12 @@ void FCurveEditorTransformTool::OnDragStart()
 					if (KeyPosition.InputValue < InputMinMax.GetLowerBoundValue())
 					{
 						InputMinMax.SetLowerBoundValue(KeyPosition.InputValue);
-		}
+					}
 					else if (KeyPosition.InputValue > InputMinMax.GetUpperBoundValue())
 					{
 						InputMinMax.SetUpperBoundValue(KeyPosition.InputValue);
 					}
-	}
+				}
 				else
 				{
 					bMinMaxIsSet = true;
@@ -1076,6 +1076,55 @@ void FCurveEditorTransformTool::OnDragEnd()
 	{
 		if (FCurveModel* Curve = CurveEditor->FindCurve(KeyData.CurveID))
 		{
+			// Gather all the final key time positions
+			TArray<FKeyPosition> KeyTimes;
+			for (int32 KeyIndex = 0; KeyIndex < KeyData.Handles.Num(); ++KeyIndex)
+			{
+				const FKeyHandle& KeyHandle = KeyData.Handles[KeyIndex];
+				FKeyPosition KeyTime = KeyData.LastDraggedKeyPositions[KeyIndex];
+				KeyTimes.Add(KeyTime);
+			}
+
+			// For each key time, look for all the keys that match
+			TArray<FKeyHandle> KeysToRemove;
+			for (const FKeyPosition& KeyTime : KeyTimes)
+			{
+				TArray<FKeyHandle> KeysInRange;
+				Curve->GetKeys(*CurveEditor, KeyTime.InputValue, KeyTime.InputValue, TNumericLimits<double>::Lowest(), TNumericLimits<double>::Max(), KeysInRange);
+
+				// If there's more than 1 key at this time, remove all but the keys that moved the largest amount
+				if (KeysInRange.Num() > 1)
+				{
+					double MaxDist = TNumericLimits<double>::Min();
+					FKeyHandle MaxKeyHandle;
+					for (const FKeyHandle& KeyInRange : KeysInRange)
+					{
+						int32 KeyDataIndex = KeyData.Handles.Find(KeyInRange);
+						if (KeyDataIndex != INDEX_NONE)
+						{
+							double Dist = FMath::Abs(KeyData.StartKeyPositions[KeyDataIndex].InputValue - KeyData.LastDraggedKeyPositions[KeyDataIndex].InputValue);
+							if (Dist > MaxDist)
+							{
+								MaxKeyHandle = KeyInRange;
+								MaxDist = Dist;
+							}
+						}
+					}
+
+					for (const FKeyHandle& KeyInRange : KeysInRange)
+					{
+						if (KeyInRange != MaxKeyHandle)
+						{
+							KeysToRemove.Add(KeyInRange);
+						}
+					}
+				}
+			}
+
+			// Remove any keys that overlap before moving new keys on top
+			Curve->RemoveKeys(KeysToRemove);
+
+			// Then, move the keys
 			Curve->SetKeyPositions(KeyData.Handles, KeyData.LastDraggedKeyPositions, EPropertyChangeType::ValueSet);
 		}
 	}
diff --git a/Engine/Source/Editor/CurveEditor/Private/DragOperations/CurveEditorDragOperation_MoveKeys.cpp b/Engine/Source/Editor/CurveEditor/Private/DragOperations/CurveEditorDragOperation_MoveKeys.cpp
--- a/Engine/Source/Editor/CurveEditor/Private/DragOperations/CurveEditorDragOperation_MoveKeys.cpp
+++ b/Engine/Source/Editor/CurveEditor/Private/DragOperations/CurveEditorDragOperation_MoveKeys.cpp
@@ -136,26 +136,53 @@ void FCurveEditorDragOperation_MoveKeys::OnEndDrag(FVector2D InitialPosition, FV
 	{
 		if (FCurveModel* Curve = CurveEditor->FindCurve(KeyData.CurveID))
 		{
-			// First, remove any keys that overlap before moving the new key on top
-			TArray<FKeyHandle> KeysToRemove;
+			// Gather all the final key time positions
+			TArray<FKeyPosition> KeyTimes;
 			for (int32 KeyIndex = 0; KeyIndex < KeyData.Handles.Num(); ++KeyIndex)
 			{
 				const FKeyHandle& KeyHandle = KeyData.Handles[KeyIndex];
 				FKeyPosition KeyTime = KeyData.LastDraggedKeyPositions[KeyIndex];
+				KeyTimes.Add(KeyTime);
+			}
 
+			// For each key time, look for all the keys that match
+			TArray<FKeyHandle> KeysToRemove;
+			for (const FKeyPosition& KeyTime : KeyTimes)
+			{
 				TArray<FKeyHandle> KeysInRange;
 				Curve->GetKeys(*CurveEditor, KeyTime.InputValue, KeyTime.InputValue, TNumericLimits<double>::Lowest(), TNumericLimits<double>::Max(), KeysInRange);
 
-				for (const FKeyHandle& KeyInRangeHandle : KeysInRange)
+				// If there's more than 1 key at this time, remove all but the keys that moved the largest amount
+				if (KeysInRange.Num() > 1)
 				{
-					if (KeyHandle != KeyInRangeHandle)
+					double MaxDist = TNumericLimits<double>::Min();
+					FKeyHandle MaxKeyHandle;
+					for (const FKeyHandle& KeyInRange : KeysInRange)
 					{
-						KeysToRemove.Add(KeyInRangeHandle);
+						int32 KeyDataIndex = KeyData.Handles.Find(KeyInRange);
+						if (KeyDataIndex != INDEX_NONE)
+						{
+							double Dist = FMath::Abs(KeyData.StartKeyPositions[KeyDataIndex].InputValue - KeyData.LastDraggedKeyPositions[KeyDataIndex].InputValue);
+							if (Dist > MaxDist)
+							{
+								MaxKeyHandle = KeyInRange;
+								MaxDist = Dist;
+							}
+						}
+					}
+
+					for (const FKeyHandle& KeyInRange : KeysInRange)
+					{
+						if (KeyInRange != MaxKeyHandle)
+						{
+							KeysToRemove.Add(KeyInRange);
+						}
 					}
 				}
 			}
-			Curve->RemoveKeys(KeysToRemove);
 
+			// Remove any keys that overlap before moving new keys on top
+			Curve->RemoveKeys(KeysToRemove);
 			// Then, move the keys
 			Curve->SetKeyPositions(KeyData.Handles, KeyData.LastDraggedKeyPositions, EPropertyChangeType::ValueSet);
 		}
