Curve Editor: Change left and right bounds for the transform tool to FFrameNumber

Add frame number customizations for tool properties panel
Moved SCurveEditorToolProperties.h to public
Fixed tool properties not getting updated on reselection

diff --git a/Engine/Plugins/Editor/CurveEditorTools/Source/CurveEditorTools/Private/Tools/CurveEditorTransformTool.cpp b/Engine/Plugins/Editor/CurveEditorTools/Source/CurveEditorTools/Private/Tools/CurveEditorTransformTool.cpp
--- a/Engine/Plugins/Editor/CurveEditorTools/Source/CurveEditorTools/Private/Tools/CurveEditorTransformTool.cpp
+++ b/Engine/Plugins/Editor/CurveEditorTools/Source/CurveEditorTools/Private/Tools/CurveEditorTransformTool.cpp
@@ -4,6 +4,7 @@
 #include "CurveEditorToolCommands.h"
 #include "CurveEditor.h"
 #include "Framework/Commands/UICommandList.h"
+#include "Misc/FrameRate.h"
 #include "Rendering/DrawElementPayloads.h"
 #include "SCurveEditorView.h"
 #include "SCurveEditorPanel.h"
@@ -308,6 +309,9 @@ void FCurveEditorTransformTool::UpdateToolOptions()
 	TSharedPtr<FCurveEditor> CurveEditor = WeakCurveEditor.Pin();
 	if (CurveEditor.IsValid())
 	{
+		const FFrameRate DisplayRate = CurveEditor->GetTimeSliderController() ? CurveEditor->GetTimeSliderController()->GetDisplayRate() : FFrameRate();
+		const FFrameRate TickResolution = CurveEditor->GetTimeSliderController() ? CurveEditor->GetTimeSliderController()->GetTickResolution() : FFrameRate();
+
 		TArray<FCurveModelID> CurveKeys;
 		CurveEditor->GetSelection().GetAll().GetKeys(CurveKeys);
 		if (CurveKeys.Num() < 1)
@@ -321,12 +325,18 @@ void FCurveEditorTransformTool::UpdateToolOptions()
 		}
 		FCurveEditorScreenSpace CurveSpace = View->GetCurveSpace(CurveKeys[0]);
 
-		ToolOptions.LeftBound = CurveSpace.ScreenToSeconds(TransformWidget.BoundsPosition.X);
+		const double LeftBoundSeconds = CurveSpace.ScreenToSeconds(TransformWidget.BoundsPosition.X);
+		const double RightBoundSeconds = CurveSpace.ScreenToSeconds(TransformWidget.BoundsPosition.X + TransformWidget.BoundsSize.X);
+
+		ToolOptions.LeftBound = FFrameRate::TransformTime(DisplayRate.AsFrameNumber(LeftBoundSeconds), DisplayRate, TickResolution).CeilToFrame();
+		ToolOptions.RightBound = FFrameRate::TransformTime(DisplayRate.AsFrameNumber(RightBoundSeconds), DisplayRate, TickResolution).CeilToFrame();
+
 		ToolOptions.UpperBound = CurveSpace.ScreenToValue(TransformWidget.BoundsPosition.Y);
-		ToolOptions.RightBound = CurveSpace.ScreenToSeconds(TransformWidget.BoundsPosition.X + TransformWidget.BoundsSize.X);
 		ToolOptions.LowerBound = CurveSpace.ScreenToValue(TransformWidget.BoundsPosition.Y + TransformWidget.BoundsSize.Y);
 
-		ToolOptions.ScaleCenterX = CurveSpace.ScreenToSeconds(TransformWidget.BoundsPosition.X + TransformWidget.BoundsSize.X * DisplayRelativeScaleCenter.X);
+		const double ScaleCenterXSeconds = CurveSpace.ScreenToSeconds(TransformWidget.BoundsPosition.X + TransformWidget.BoundsSize.X * DisplayRelativeScaleCenter.X);
+
+		ToolOptions.ScaleCenterX = FFrameRate::TransformTime(DisplayRate.AsFrameNumber(ScaleCenterXSeconds), DisplayRate, TickResolution).CeilToFrame();
 		ToolOptions.ScaleCenterY = CurveSpace.ScreenToValue(TransformWidget.BoundsPosition.Y + TransformWidget.BoundsSize.Y * DisplayRelativeScaleCenter.Y);
 	}
 }
@@ -444,6 +454,9 @@ void FCurveEditorTransformTool::OnToolOptionsUpdated(const FPropertyChangedEvent
 		return;
 	}
 
+	const FFrameRate DisplayRate = CurveEditor->GetTimeSliderController() ? CurveEditor->GetTimeSliderController()->GetDisplayRate() : FFrameRate();
+	const FFrameRate TickResolution = CurveEditor->GetTimeSliderController() ? CurveEditor->GetTimeSliderController()->GetTickResolution() : FFrameRate();
+
 	FCurveEditorScreenSpace CurveSpace = View->GetCurveSpace(CurveKeys[0]);
 
 	FVector2D ScaleCenter = FVector2D(0.5f, 0.5f);
@@ -454,21 +467,21 @@ void FCurveEditorTransformTool::OnToolOptionsUpdated(const FPropertyChangedEvent
 	{
 		bAffectsY = false;
 		ScaleCenter.X = 1.0f;
-		ScaleDelta.X = -(CurveSpace.SecondsToScreen(ToolOptions.LeftBound) - TransformWidget.Position.X);
+		const double LeftBoundSeconds = DisplayRate.AsSeconds(FFrameRate::TransformTime(ToolOptions.LeftBound, TickResolution, DisplayRate));
+		ScaleDelta.X = -(CurveSpace.SecondsToScreen(LeftBoundSeconds) - TransformWidget.Position.X);
 	}
 	else if (PropertyChangedEvent.GetPropertyName().IsEqual(GET_MEMBER_NAME_CHECKED(FTransformToolOptions, UpperBound)))
 	{
 		bAffectsX = false;
 		ScaleCenter.Y = 1.0f;
 		ScaleDelta.Y = -(CurveSpace.ValueToScreen(ToolOptions.UpperBound) - TransformWidget.Position.Y);
-		float diff = ToolOptions.UpperBound - CurveSpace.ScreenToValue(TransformWidget.Position.Y);
-		UE_LOG(LogTemp, Log, TEXT("diff: %f"), diff);
 	}
 	else if (PropertyChangedEvent.GetPropertyName().IsEqual(GET_MEMBER_NAME_CHECKED(FTransformToolOptions, RightBound)))
 	{
 		bAffectsY = false;
 		ScaleCenter.X = 0.0f;
-		ScaleDelta.X = CurveSpace.SecondsToScreen(ToolOptions.RightBound) - (TransformWidget.Position.X + TransformWidget.Size.X);
+		const double RightBoundSeconds = DisplayRate.AsSeconds(FFrameRate::TransformTime(ToolOptions.RightBound, TickResolution, DisplayRate));
+		ScaleDelta.X = CurveSpace.SecondsToScreen(RightBoundSeconds) - (TransformWidget.Position.X + TransformWidget.Size.X);
 	}
 	else if (PropertyChangedEvent.GetPropertyName().IsEqual(GET_MEMBER_NAME_CHECKED(FTransformToolOptions, LowerBound)))
 	{
@@ -478,7 +491,8 @@ void FCurveEditorTransformTool::OnToolOptionsUpdated(const FPropertyChangedEvent
 	}
 	else if (PropertyChangedEvent.GetPropertyName().IsEqual(GET_MEMBER_NAME_CHECKED(FTransformToolOptions, ScaleCenterX)))
 	{
-		const double ViewSpaceX = CurveSpace.SecondsToScreen(ToolOptions.ScaleCenterX);
+		const double ScaleCenterXSeconds = DisplayRate.AsSeconds(FFrameRate::TransformTime(ToolOptions.ScaleCenterX, TickResolution, DisplayRate));
+		const double ViewSpaceX = CurveSpace.SecondsToScreen(ScaleCenterXSeconds);
 		const double ViewSpaceXDelta = ViewSpaceX - TransformWidget.BoundsPosition.X;
 		RelativeScaleCenter.X = ViewSpaceXDelta / TransformWidget.BoundsSize.X;
 	}
@@ -1236,6 +1250,9 @@ void FCurveEditorTransformTool::Tick(const FGeometry& AllottedGeometry, const do
 		return;
 	}
 
+	const FFrameRate DisplayRate = CurveEditor->GetTimeSliderController() ? CurveEditor->GetTimeSliderController()->GetDisplayRate() : FFrameRate();
+	const FFrameRate TickResolution = CurveEditor->GetTimeSliderController() ? CurveEditor->GetTimeSliderController()->GetTickResolution() : FFrameRate();
+
 	// Check if all curve models have the same scales
 	TArray<FCurveModelID> CurveModelIDs;
 	CurveEditor->GetSelection().GetAll().GetKeys(CurveModelIDs);
@@ -1277,11 +1294,17 @@ void FCurveEditorTransformTool::Tick(const FGeometry& AllottedGeometry, const do
 
 		FCurveEditorScreenSpace CurveSpace = View->GetCurveSpace(FirstCurveID);
 
-		const FVector2D OptionsSize = FVector2D(ToolOptions.RightBound - ToolOptions.LeftBound, ToolOptions.LowerBound - ToolOptions.UpperBound);
+		const double RightBoundSeconds = DisplayRate.AsSeconds(FFrameRate::TransformTime(ToolOptions.RightBound, TickResolution, DisplayRate));
+		const double LeftBoundSeconds = DisplayRate.AsSeconds(FFrameRate::TransformTime(ToolOptions.LeftBound, TickResolution, DisplayRate));
+
+		const FVector2D OptionsSize = FVector2D(RightBoundSeconds - LeftBoundSeconds, ToolOptions.LowerBound - ToolOptions.UpperBound);
 		const FVector2D ViewSpaceOptionsSize = FVector2D(CurveSpace.PixelsPerInput() * FMath::Abs(OptionsSize.X), CurveSpace.PixelsPerOutput() * FMath::Abs(OptionsSize.Y));
-		if ((!DelayedDrag || (DelayedDrag && !DelayedDrag->IsDragging())) && (TransformWidget.BoundsSize - ViewSpaceOptionsSize).IsNearlyZero())
+		if ((!DelayedDrag || (DelayedDrag && !DelayedDrag->IsDragging())) && 
+			((TransformWidget.BoundsSize - ViewSpaceOptionsSize).IsNearlyZero() || SelectionSerialNumber != CurveEditor->GetSelection().GetSerialNumber()))
 		{
 			UpdateToolOptions();
+
+			SelectionSerialNumber = CurveEditor->GetSelection().GetSerialNumber();
 		}
 	}
 }
diff --git a/Engine/Plugins/Editor/CurveEditorTools/Source/CurveEditorTools/Private/Tools/CurveEditorTransformTool.h b/Engine/Plugins/Editor/CurveEditorTools/Source/CurveEditorTools/Private/Tools/CurveEditorTransformTool.h
--- a/Engine/Plugins/Editor/CurveEditorTools/Source/CurveEditorTools/Private/Tools/CurveEditorTransformTool.h
+++ b/Engine/Plugins/Editor/CurveEditorTools/Source/CurveEditorTools/Private/Tools/CurveEditorTransformTool.h
@@ -93,13 +93,13 @@ struct FTransformToolOptions
 	float LowerBound = 0.f;
 
 	UPROPERTY(Transient, EditAnywhere, Category = ToolOptions)
-	float LeftBound = 0.f;
+	FFrameNumber LeftBound;
 
 	UPROPERTY(Transient, EditAnywhere, Category = ToolOptions)
-	float RightBound = 0.f;
+	FFrameNumber RightBound;
 
 	UPROPERTY(Transient, EditAnywhere, Category = ToolOptions)
-	float ScaleCenterX = 0.f;
+	FFrameNumber ScaleCenterX;
 
 	UPROPERTY(Transient, EditAnywhere, Category = ToolOptions)
 	float ScaleCenterY = 0.f;
@@ -209,5 +209,8 @@ private:
 
 	/** Modify the Falloff Weight by type of Interpolation */
 	double ModifyWeightByInterpType(double Weight) const;
+
+	/** Selection serial number, used to update the toolbox */
+	uint32 SelectionSerialNumber = 0;
 };
 
diff --git a/Engine/Source/Editor/CurveEditor/Private/SCurveEditorToolProperties.h b/Engine/Source/Editor/CurveEditor/Private/SCurveEditorToolProperties.h
deleted file mode 100644
diff --git a/Engine/Source/Editor/CurveEditor/Public/SCurveEditorPanel.h b/Engine/Source/Editor/CurveEditor/Public/SCurveEditorPanel.h
--- a/Engine/Source/Editor/CurveEditor/Public/SCurveEditorPanel.h
+++ b/Engine/Source/Editor/CurveEditor/Public/SCurveEditorPanel.h
@@ -109,6 +109,14 @@ class CURVEEDITOR_API SCurveEditorPanel : public SCompoundWidget
 		return FilterPanel;
 	}
 
+	/**
+	 * Access the tool properties panel
+	 */
+	TSharedPtr<class SCurveEditorToolProperties> GetToolPropertiesPanel() const
+	{
+		return ToolPropertiesPanel;
+	}
+
 	void AddView(TSharedRef<SCurveEditorView> ViewToAdd);
 
 	void RemoveView(TSharedRef<SCurveEditorView> ViewToRemove);
diff --git a/Engine/Source/Editor/CurveEditor/Public/SCurveEditorToolProperties.h b/Engine/Source/Editor/CurveEditor/Public/SCurveEditorToolProperties.h
new file mode 100644
--- /dev/null
+++ b/Engine/Source/Editor/CurveEditor/Public/SCurveEditorToolProperties.h
@@ -0,0 +1,56 @@
+// Copyright Epic Games, Inc. All Rights Reserved.
+
+#pragma once
+
+#include "ICurveEditorToolExtension.h"
+#include "Templates/SharedPointer.h"
+#include "Widgets/DeclarativeSyntaxSupport.h"
+#include "Widgets/SCompoundWidget.h"
+
+class FCurveEditor;
+class IStructureDetailsView;
+struct FPropertyChangedEvent;
+
+/**
+ * Widget to show curve editor tool extension properties
+ */
+class CURVEEDITOR_API SCurveEditorToolProperties : public SCompoundWidget
+{
+public:
+	SLATE_BEGIN_ARGS(SCurveEditorToolProperties)
+	{}
+
+	SLATE_END_ARGS()
+
+	SCurveEditorToolProperties() : ToolId(FCurveEditorToolID::Unset()) {}
+
+	void Construct(const FArguments& InArgs, TSharedRef<FCurveEditor> InCurveEditor, FCurveEditorToolID ToolId);
+
+	/**
+	 * Called when the curve editor tool changes
+	 * Sets the new tool ID and rebuilds the details view
+	 */
+	void OnToolChanged(FCurveEditorToolID NewToolId);
+
+	/** The details view for tool properties */
+	TSharedPtr<class IStructureDetailsView> GetStructureDetailsView() const { return DetailsView; }
+
+private:
+
+	/**
+	 * Rebuilds details view
+	 * Called either by the tool itself when it's properties change or when the curve editor tool changes
+	 */
+	void RebuildProperties();
+
+	/**
+	 * Called when the user enters in a new value for one of the properties
+	 * Calls OnToolOptionsUpdated on the current tool
+	 */
+	void OnFinishedChangingProperties(const FPropertyChangedEvent& PropertyChangedEvent);
+
+private:
+	TWeakPtr<FCurveEditor> WeakCurveEditor;
+	TSharedPtr<IStructureDetailsView> DetailsView;
+	FCurveEditorToolID ToolId;
+};
\ No newline at end of file
diff --git a/Engine/Source/Editor/Sequencer/Private/MVVM/CurveEditorExtension.cpp b/Engine/Source/Editor/Sequencer/Private/MVVM/CurveEditorExtension.cpp
--- a/Engine/Source/Editor/Sequencer/Private/MVVM/CurveEditorExtension.cpp
+++ b/Engine/Source/Editor/Sequencer/Private/MVVM/CurveEditorExtension.cpp
@@ -7,8 +7,10 @@
 #include "Filters/SCurveEditorFilterPanel.h"
 #include "Framework/Docking/TabManager.h"
 #include "IPropertyRowGenerator.h"
+#include "IStructureDetailsView.h"
 #include "MVVM/ViewModels/SequencerEditorViewModel.h"
 #include "SCurveEditorPanel.h"
+#include "SCurveEditorToolProperties.h"
 #include "SCurveKeyDetailPanel.h"
 #include "SSequencerTreeFilterStatusBar.h"
 #include "Sequencer.h"
@@ -246,6 +248,9 @@ void FCurveEditorExtension::CreateCurveEditor(const FTimeSliderArgs& TimeSliderA
 	CurveEditorPanel->GetKeyDetailsView()->GetPropertyRowGenerator()->RegisterInstancedCustomPropertyTypeLayout(
 			"FrameNumber", 
 			FOnGetPropertyTypeCustomizationInstance::CreateStatic(CreateFrameNumberCustomization, WeakSequencer));
+	CurveEditorPanel->GetToolPropertiesPanel()->GetStructureDetailsView()->GetDetailsView()->RegisterInstancedCustomPropertyTypeLayout(
+		"FrameNumber",
+		FOnGetPropertyTypeCustomizationInstance::CreateStatic(CreateFrameNumberCustomization, WeakSequencer));
 
 	// And jump to the Curve Editor tree search if you have the Curve Editor focused
 	CurveEditorModel->GetCommands()->MapAction(
