Curve Editor: Fix BakeInterval for FFT Filter

diff --git a/Engine/Plugins/Editor/CurveEditorTools/Source/CurveEditorTools/Private/Filters/CurveEditorFFTFilter.cpp b/Engine/Plugins/Editor/CurveEditorTools/Source/CurveEditorTools/Private/Filters/CurveEditorFFTFilter.cpp
--- a/Engine/Plugins/Editor/CurveEditorTools/Source/CurveEditorTools/Private/Filters/CurveEditorFFTFilter.cpp
+++ b/Engine/Plugins/Editor/CurveEditorTools/Source/CurveEditorTools/Private/Filters/CurveEditorFFTFilter.cpp
@@ -65,10 +65,10 @@ void UCurveEditorFFTFilter::ApplyFilter_Impl(TSharedRef<FCurveEditor> InCurveEdi
 		}
 		//Set interval range divded by twice the number of keys.
 
-		const double IntervalAsSeconds = (MaxKey - MinKey) / (2 * OriginalKeyHandles.Num());
+		double IntervalAsSeconds = (MaxKey - MinKey) / (2 * OriginalKeyHandles.Num());
 		BakeFilter->BakeIntervalInSeconds = IntervalAsSeconds;
-		BakeFilter->BakeInterval = FFrameRate::TransformTime(DisplayRate.AsFrameNumber(IntervalAsSeconds), DisplayRate, TickResolution).FrameNumber;
-		
+		BakeFilter->BakeInterval = (IntervalAsSeconds * TickResolution).FloorToFrame();
+
 		// This stores the position of all of the original keys. Once we've filtered the curve we will need to sample it at these positions.
 		TArray<FKeyPosition> OriginalKeyPositions;
 		TArray<FKeyAttributes> OriginalKeyAttributes;
