Curve Editor: Use member property name instead of property name when responding to changes from FTransformToolOptions

diff --git a/Engine/Plugins/Editor/CurveEditorTools/Source/CurveEditorTools/Private/Tools/CurveEditorTransformTool.cpp b/Engine/Plugins/Editor/CurveEditorTools/Source/CurveEditorTools/Private/Tools/CurveEditorTransformTool.cpp
--- a/Engine/Plugins/Editor/CurveEditorTools/Source/CurveEditorTools/Private/Tools/CurveEditorTransformTool.cpp
+++ b/Engine/Plugins/Editor/CurveEditorTools/Source/CurveEditorTools/Private/Tools/CurveEditorTransformTool.cpp
@@ -463,40 +463,40 @@ void FCurveEditorTransformTool::OnToolOptionsUpdated(const FPropertyChangedEvent
 	FVector2D ScaleDelta = FVector2D(0.0f, 0.0f);
 	bool bAffectsX = true, bAffectsY = true;
 
-	if (PropertyChangedEvent.GetPropertyName().IsEqual(GET_MEMBER_NAME_CHECKED(FTransformToolOptions, LeftBound)))
+	if (PropertyChangedEvent.GetMemberPropertyName().IsEqual(GET_MEMBER_NAME_CHECKED(FTransformToolOptions, LeftBound)))
 	{
 		bAffectsY = false;
 		ScaleCenter.X = 1.0f;
 		const double LeftBoundSeconds = DisplayRate.AsSeconds(FFrameRate::TransformTime(ToolOptions.LeftBound, TickResolution, DisplayRate));
 		ScaleDelta.X = -(CurveSpace.SecondsToScreen(LeftBoundSeconds) - TransformWidget.Position.X);
 	}
-	else if (PropertyChangedEvent.GetPropertyName().IsEqual(GET_MEMBER_NAME_CHECKED(FTransformToolOptions, UpperBound)))
+	else if (PropertyChangedEvent.GetMemberPropertyName().IsEqual(GET_MEMBER_NAME_CHECKED(FTransformToolOptions, UpperBound)))
 	{
 		bAffectsX = false;
 		ScaleCenter.Y = 1.0f;
 		ScaleDelta.Y = -(CurveSpace.ValueToScreen(ToolOptions.UpperBound) - TransformWidget.Position.Y);
 	}
-	else if (PropertyChangedEvent.GetPropertyName().IsEqual(GET_MEMBER_NAME_CHECKED(FTransformToolOptions, RightBound)))
+	else if (PropertyChangedEvent.GetMemberPropertyName().IsEqual(GET_MEMBER_NAME_CHECKED(FTransformToolOptions, RightBound)))
 	{
 		bAffectsY = false;
 		ScaleCenter.X = 0.0f;
 		const double RightBoundSeconds = DisplayRate.AsSeconds(FFrameRate::TransformTime(ToolOptions.RightBound, TickResolution, DisplayRate));
 		ScaleDelta.X = CurveSpace.SecondsToScreen(RightBoundSeconds) - (TransformWidget.Position.X + TransformWidget.Size.X);
 	}
-	else if (PropertyChangedEvent.GetPropertyName().IsEqual(GET_MEMBER_NAME_CHECKED(FTransformToolOptions, LowerBound)))
+	else if (PropertyChangedEvent.GetMemberPropertyName().IsEqual(GET_MEMBER_NAME_CHECKED(FTransformToolOptions, LowerBound)))
 	{
 		bAffectsX = false;
 		ScaleCenter.Y = 0.0f;
 		ScaleDelta.Y = CurveSpace.ValueToScreen(ToolOptions.LowerBound) - (TransformWidget.Position.Y + TransformWidget.Size.Y);
 	}
-	else if (PropertyChangedEvent.GetPropertyName().IsEqual(GET_MEMBER_NAME_CHECKED(FTransformToolOptions, ScaleCenterX)))
+	else if (PropertyChangedEvent.GetMemberPropertyName().IsEqual(GET_MEMBER_NAME_CHECKED(FTransformToolOptions, ScaleCenterX)))
 	{
 		const double ScaleCenterXSeconds = DisplayRate.AsSeconds(FFrameRate::TransformTime(ToolOptions.ScaleCenterX, TickResolution, DisplayRate));
 		const double ViewSpaceX = CurveSpace.SecondsToScreen(ScaleCenterXSeconds);
 		const double ViewSpaceXDelta = ViewSpaceX - TransformWidget.BoundsPosition.X;
 		RelativeScaleCenter.X = ViewSpaceXDelta / TransformWidget.BoundsSize.X;
 	}
-	else if (PropertyChangedEvent.GetPropertyName().IsEqual(GET_MEMBER_NAME_CHECKED(FTransformToolOptions, ScaleCenterY)))
+	else if (PropertyChangedEvent.GetMemberPropertyName().IsEqual(GET_MEMBER_NAME_CHECKED(FTransformToolOptions, ScaleCenterY)))
 	{
 		const double ViewSpaceY = CurveSpace.ValueToScreen(ToolOptions.ScaleCenterY);
 		const double ViewSpaceYDelta = ViewSpaceY - TransformWidget.BoundsPosition.Y;
