Curve Editor: Focus playback time/range commands should just be regular buttons and not radio buttons

diff --git a/Engine/Plugins/Editor/CurveEditorTools/Source/CurveEditorTools/Private/CurveEditorToolCommands.cpp b/Engine/Plugins/Editor/CurveEditorTools/Source/CurveEditorTools/Private/CurveEditorToolCommands.cpp
--- a/Engine/Plugins/Editor/CurveEditorTools/Source/CurveEditorTools/Private/CurveEditorToolCommands.cpp
+++ b/Engine/Plugins/Editor/CurveEditorTools/Source/CurveEditorTools/Private/CurveEditorToolCommands.cpp
@@ -7,8 +7,8 @@
 void FCurveEditorToolCommands::RegisterCommands()
 {
 	// Focus Tools
-	UI_COMMAND(SetFocusPlaybackTime, "Focus Playback Time", "Focuses the Curve Editor on the current Playback Time without changing zoom level.", EUserInterfaceActionType::RadioButton, FInputChord());
-	UI_COMMAND(SetFocusPlaybackRange, "Focus Playback Range", "Focuses the Curve Editor on the current Playback Range with zoom based on visible curves.", EUserInterfaceActionType::RadioButton, FInputChord(EKeys::A));
+	UI_COMMAND(SetFocusPlaybackTime, "Focus Playback Time", "Focuses the Curve Editor on the current Playback Time without changing zoom level.", EUserInterfaceActionType::Button, FInputChord());
+	UI_COMMAND(SetFocusPlaybackRange, "Focus Playback Range", "Focuses the Curve Editor on the current Playback Range with zoom based on visible curves.", EUserInterfaceActionType::Button, FInputChord(EKeys::A));
 
 	// Tool Modes
 	UI_COMMAND(ActivateTransformTool, "Transform Tool", "Activates the Transform tool which allows translation, scale and rotation of selected keys.", EUserInterfaceActionType::ToggleButton, FInputChord(EModifierKey::Control, EKeys::T));
