Sequencer: Curve Editor: Speculative fix for hard to reproduce crash when no keys are selected with multi-select tool.

diff --git a/Engine/Plugins/Editor/CurveEditorTools/Source/CurveEditorTools/Private/Tools/CurveEditorMultiScaleTool.cpp b/Engine/Plugins/Editor/CurveEditorTools/Source/CurveEditorTools/Private/Tools/CurveEditorMultiScaleTool.cpp
--- a/Engine/Plugins/Editor/CurveEditorTools/Source/CurveEditorTools/Private/Tools/CurveEditorMultiScaleTool.cpp
+++ b/Engine/Plugins/Editor/CurveEditorTools/Source/CurveEditorTools/Private/Tools/CurveEditorMultiScaleTool.cpp
@@ -141,7 +141,7 @@ void FCurveEditorMultiScaleTool::OnPaint(const FPaintArgs& Args, const FGeometry
 			TArray<FKeyPosition> KeyPositions;
 			KeyPositions.SetNumZeroed(Handles.Num());
 			Curve->GetKeyPositions(Handles, KeyPositions);
-
+			
 			const FVector2D Pivot = DelayedDrag.IsSet() && DelayedDrag->IsDragging() ? KeysByCurve[KeysByCurveIdx].Pivot : GetPivot(Curve, KeyPositions);
 
 			FCurveEditorScreenSpace CurveSpace = View->GetCurveSpace(CurveID);
@@ -389,6 +389,10 @@ void FCurveEditorMultiScaleTool::OnDragStart()
 
 			KeyData.StartKeyPositions.SetNumZeroed(KeyData.Handles.Num());
 			Curve->GetKeyPositions(KeyData.Handles, KeyData.StartKeyPositions);
+			if (KeyData.StartKeyPositions.Num() == 0)
+			{
+				continue;
+			}
 			for (const FKeyPosition& KeyPosition : KeyData.StartKeyPositions)
 			{
 				if (bMinMaxIsSet)
@@ -632,7 +636,10 @@ void FCurveEditorMultiScaleTool::ScaleUnique(const FVector2D& InChangeAmount, co
 FVector2D FCurveEditorMultiScaleTool::GetPivot(FCurveModel* InCurve, const TArray<FKeyPosition>& InKeyPositions) const
 {
 	FVector2D Pivot = FVector2D::ZeroVector;
-	check(InKeyPositions.Num() != 0);
+	if (InKeyPositions.Num() == 0)
+	{
+		return Pivot;
+	}
 
 	switch (ToolOptions.PivotType)
 	{
