Curve Editor: Make sure to modify the curve in the retime tool before setting key positions. This fixes undo with the retime tool.

diff --git a/Engine/Plugins/Editor/CurveEditorTools/Source/CurveEditorTools/Private/Tools/CurveEditorRetimeTool.cpp b/Engine/Plugins/Editor/CurveEditorTools/Source/CurveEditorTools/Private/Tools/CurveEditorRetimeTool.cpp
--- a/Engine/Plugins/Editor/CurveEditorTools/Source/CurveEditorTools/Private/Tools/CurveEditorRetimeTool.cpp
+++ b/Engine/Plugins/Editor/CurveEditorTools/Source/CurveEditorTools/Private/Tools/CurveEditorRetimeTool.cpp
@@ -468,6 +468,8 @@ void FCurveEditorRetimeTool::OnDragStart()
 			FCurveModel* CurveModel = CurveEditor->FindCurve(CurveID);
 			check(CurveModel);
 
+			CurveModel->Modify();
+
 			FPreDragChannelData& ChannelData = ResizeData.CurveChannels[ResizeData.CurveChannels.Emplace(CurveID)];
 
 			// Get all Key Handles
