Curve Editor Tools: Retime tool now allows for scrubbing

diff --git a/Engine/Plugins/Editor/CurveEditorTools/Source/CurveEditorTools/Private/Tools/CurveEditorRetimeTool.cpp b/Engine/Plugins/Editor/CurveEditorTools/Source/CurveEditorTools/Private/Tools/CurveEditorRetimeTool.cpp
--- a/Engine/Plugins/Editor/CurveEditorTools/Source/CurveEditorTools/Private/Tools/CurveEditorRetimeTool.cpp
+++ b/Engine/Plugins/Editor/CurveEditorTools/Source/CurveEditorTools/Private/Tools/CurveEditorRetimeTool.cpp
@@ -3,6 +3,7 @@
 #include "CurveEditorRetimeTool.h"
 #include "CurveEditorToolCommands.h"
 #include "CurveEditor.h"
+#include "CurveEditorCommands.h"
 #include "Framework/Commands/UICommandList.h"
 #include "ICurveEditorBounds.h"
 #include "Rendering/DrawElementPayloads.h"
@@ -59,6 +60,20 @@ FCurveEditorRetimeTool::~FCurveEditorRetimeTool()
 	RetimeData = nullptr;
 }
 
+FReply FCurveEditorRetimeTool::OnKeyDown(TSharedRef<SWidget> OwningWidget, const FGeometry& MyGeometry, const FKeyEvent& InKeyEvent)
+{
+	bScrubbing = IsScrubTimeKeyEvent(InKeyEvent);
+
+	return FReply::Unhandled();
+}
+
+FReply FCurveEditorRetimeTool::OnKeyUp(TSharedRef<SWidget> OwningWidget, const FGeometry& MyGeometry, const FKeyEvent& InKeyEvent)
+{
+	bScrubbing = false;
+
+	return FReply::Unhandled();
+}
+
 FReply FCurveEditorRetimeTool::OnMouseButtonDown(TSharedRef<SWidget> OwningWidget, const FGeometry& MyGeometry, const FPointerEvent& MouseEvent)
 {
 	// UE_LOG(LogTemp, Log, TEXT("FCurveEditorRetimeTool::OnMouseButtonDown"));
@@ -132,8 +147,12 @@ FReply FCurveEditorRetimeTool::OnMouseButtonDown(TSharedRef<SWidget> OwningWidge
 			DelayedDrag = FDelayedDrag(MyGeometry.AbsoluteToLocal(MouseEvent.GetScreenSpacePosition()), MouseEvent.GetEffectingButton());
 		}
 
-		// We always handle the left mouse button so clicks don't bubble through to selection
-		return FReply::Handled().PreventThrottling();
+		// Let scrubbing bubble through to the parent widget
+		if (!bScrubbing)
+		{
+			// Handle the left mouse button so clicks don't bubble through to selection.
+			return FReply::Handled().PreventThrottling();
+		}
 	}
 
 	// Right Clicks need to go through so that panning the graph works
@@ -605,6 +624,7 @@ void FCurveEditorRetimeTool::OnDrag(const FVector2D& InStartPosition, const FVec
 void FCurveEditorRetimeTool::OnDragEnd()
 {
 	TSharedPtr<FCurveEditor> CurveEditor = WeakCurveEditor.Pin();
+
 	// For each curve we're affecting.
 	for (FPreDragChannelData& ChannelData : PreDragCurveData->CurveChannels)
 	{
@@ -638,5 +658,22 @@ void FCurveEditorRetimeTool::StopDragIfPossible()
 	}
 }
 
+bool FCurveEditorRetimeTool::IsScrubTimeKeyEvent(const FKeyEvent& InKeyEvent) const
+{
+	const FCurveEditorCommands& Commands = FCurveEditorCommands::Get();
+	// Need to iterate through primary and secondary to make sure they are all pressed.
+	for (uint32 i = 0; i < static_cast<uint8>(EMultipleKeyBindingIndex::NumChords); ++i)
+	{
+		EMultipleKeyBindingIndex ChordIndex = static_cast<EMultipleKeyBindingIndex>(i);
+		const FInputChord& Chord = *Commands.ScrubTime->GetActiveChord(ChordIndex);
+		const bool bIsMovingTimeSlider = Chord.IsValidChord() && InKeyEvent.GetKey() == Chord.Key;
+		if (bIsMovingTimeSlider)
+		{
+			return true;
+		}
+	}
+	return false;
+}
+
 #undef LOCTEXT_NAMESPACE // "CurveEditorToolCommands"
 
diff --git a/Engine/Plugins/Editor/CurveEditorTools/Source/CurveEditorTools/Private/Tools/CurveEditorRetimeTool.h b/Engine/Plugins/Editor/CurveEditorTools/Source/CurveEditorTools/Private/Tools/CurveEditorRetimeTool.h
--- a/Engine/Plugins/Editor/CurveEditorTools/Source/CurveEditorTools/Private/Tools/CurveEditorRetimeTool.h
+++ b/Engine/Plugins/Editor/CurveEditorTools/Source/CurveEditorTools/Private/Tools/CurveEditorRetimeTool.h
@@ -70,6 +70,8 @@ public:
 
 	// ICurveEditorToolExtension
 	virtual void OnPaint(const FPaintArgs& Args, const FGeometry& AllottedGeometry, const FSlateRect& MyCullingRect, FSlateWindowElementList& OutDrawElements, int32 PaintOnLayerId, const FWidgetStyle& InWidgetStyle, bool bParentEnabled) const override;
+	virtual FReply OnKeyDown(TSharedRef<SWidget> OwningWidget, const FGeometry& MyGeometry, const FKeyEvent& InKeyEvent) override;
+	virtual FReply OnKeyUp(TSharedRef<SWidget> OwningWidget, const FGeometry& MyGeometry, const FKeyEvent& InKeyEvent) override;
 	virtual FReply OnMouseButtonDown(TSharedRef<SWidget> OwningWidget, const FGeometry& MyGeometry, const FPointerEvent& MouseEvent) override;
 	virtual FReply OnMouseButtonDoubleClick(TSharedRef<SWidget> OwningWidget, const FGeometry& InMyGeometry, const FPointerEvent& InMouseEvent) override;
 	virtual FReply OnMouseButtonUp(TSharedRef<SWidget> OwningWidget, const FGeometry& MyGeometry, const FPointerEvent& MouseEvent) override;
@@ -85,8 +87,12 @@ private:
 	virtual void OnDragStart();
 	virtual void OnDrag(const FVector2D& InStartPosition, const FVector2D& InMousePosition);
 	virtual void OnDragEnd();
+
 	void StopDragIfPossible();
 
+	/** Returns true if the key event is a scrub time key event */
+	bool IsScrubTimeKeyEvent(const FKeyEvent& InKeyEvent) const;
+
 private:
 	/** Weak pointer back to the owning curve editor. */
 	TWeakPtr<FCurveEditor> WeakCurveEditor;
@@ -130,4 +136,7 @@ private:
 
 	/** Cached user selection for keys. We clear this when the tool is opened and restore it when they leave */
 	TMap<FCurveModelID, FKeyHandleSet> CachedSelectionSet;
+
+	/** True while scrubbing */
+	bool bScrubbing = false;
 };
diff --git a/Engine/Source/Editor/CurveEditor/Private/SCurveEditorViewContainer.cpp b/Engine/Source/Editor/CurveEditor/Private/SCurveEditorViewContainer.cpp
--- a/Engine/Source/Editor/CurveEditor/Private/SCurveEditorViewContainer.cpp
+++ b/Engine/Source/Editor/CurveEditor/Private/SCurveEditorViewContainer.cpp
@@ -180,24 +180,40 @@ bool SCurveEditorViewContainer::IsScrubTimeKeyEvent(const FKeyEvent& InKeyEvent)
 
 FReply SCurveEditorViewContainer::OnKeyDown(const FGeometry& MyGeometry, const FKeyEvent& InKeyEvent)
 {
-	if (InKeyEvent.GetKey() == EKeys::Escape && DragOperation.IsSet())
+	FReply ToolReply = FReply::Unhandled();
+	if (CurveEditor->GetCurrentTool())
+	{
+		ToolReply = CurveEditor->GetCurrentTool()->OnKeyDown(AsShared(), MyGeometry, InKeyEvent);
+	}
+
+	if (!ToolReply.IsEventHandled() &&
+		InKeyEvent.GetKey() == EKeys::Escape && DragOperation.IsSet())
 	{
 		DragOperation->DragImpl->CancelDrag();
 		DragOperation.Reset();
 		return FReply::Handled();
 	}
 
-	if (IsScrubTimeKeyEvent(InKeyEvent))
+	if (!ToolReply.IsEventHandled() && 
+		IsScrubTimeKeyEvent(InKeyEvent))
 	{
 		bIsScrubbingTime = true;
 		return FReply::Handled();
 	}
-	return FReply::Unhandled();
+
+	return ToolReply;
 }
 
 FReply SCurveEditorViewContainer::OnKeyUp(const FGeometry& MyGeometry, const FKeyEvent& InKeyEvent)
 {
-	if (IsScrubTimeKeyEvent(InKeyEvent) && bIsScrubbingTime )
+	FReply ToolReply = FReply::Unhandled();
+	if (CurveEditor->GetCurrentTool())
+	{
+		ToolReply = CurveEditor->GetCurrentTool()->OnKeyUp(AsShared(), MyGeometry, InKeyEvent);
+	}
+
+	if (!ToolReply.IsEventHandled() && 
+		IsScrubTimeKeyEvent(InKeyEvent) && bIsScrubbingTime )
 	{
 		bIsScrubbingTime = false;
 		if (DragOperation.IsSet())
@@ -208,7 +224,7 @@ FReply SCurveEditorViewContainer::OnKeyUp(const FGeometry& MyGeometry, const FKe
 		}
 	}
 	
-	return FReply::Unhandled();
+	return ToolReply;
 }
 
 FReply SCurveEditorViewContainer::OnPreviewMouseButtonDown(const FGeometry& MyGeometry, const FPointerEvent& MouseEvent)
diff --git a/Engine/Source/Editor/CurveEditor/Public/ICurveEditorToolExtension.h b/Engine/Source/Editor/CurveEditor/Public/ICurveEditorToolExtension.h
--- a/Engine/Source/Editor/CurveEditor/Public/ICurveEditorToolExtension.h
+++ b/Engine/Source/Editor/CurveEditor/Public/ICurveEditorToolExtension.h
@@ -91,6 +91,8 @@ public:
 	// Effectively SWidget Interface
 	virtual void OnPaint(const FPaintArgs& Args, const FGeometry& AllottedGeometry, const FSlateRect& MyCullingRect, FSlateWindowElementList& OutDrawElements, int32 PaintOnLayerId, const FWidgetStyle& InWidgetStyle, bool bParentEnabled) const = 0;
 	virtual void Tick(const FGeometry& AllottedGeometry, const double InCurrentTime, const float InDeltaTime) {}
+	virtual FReply OnKeyDown(TSharedRef<SWidget> OwningWidget, const FGeometry& MyGeometry, const FKeyEvent& InKeyEvent) { return FReply::Unhandled(); }
+	virtual FReply OnKeyUp(TSharedRef<SWidget> OwningWidget, const FGeometry& MyGeometry, const FKeyEvent& InKeyEvent) { return FReply::Unhandled(); }
 	virtual FReply OnMouseButtonDown(TSharedRef<SWidget> OwningWidget, const FGeometry& MyGeometry, const FPointerEvent& MouseEvent) { return FReply::Unhandled(); }
 	virtual FReply OnMouseButtonUp(TSharedRef<SWidget> OwningWidget, const FGeometry& MyGeometry, const FPointerEvent& MouseEvent) { return FReply::Unhandled(); }
 	virtual FReply OnMouseMove(TSharedRef<SWidget> OwningWidget, const FGeometry& MyGeometry, const FPointerEvent& MouseEvent) { return FReply::Unhandled(); }
