Curve Editor Tools: Adding Retime Anchors now transacts to the undo stack

diff --git a/Engine/Plugins/Editor/CurveEditorTools/Source/CurveEditorTools/Private/Tools/CurveEditorRetimeTool.cpp b/Engine/Plugins/Editor/CurveEditorTools/Source/CurveEditorTools/Private/Tools/CurveEditorRetimeTool.cpp
--- a/Engine/Plugins/Editor/CurveEditorTools/Source/CurveEditorTools/Private/Tools/CurveEditorRetimeTool.cpp
+++ b/Engine/Plugins/Editor/CurveEditorTools/Source/CurveEditorTools/Private/Tools/CurveEditorRetimeTool.cpp
@@ -232,6 +232,9 @@ FReply FCurveEditorRetimeTool::OnMouseButtonDoubleClick(TSharedRef<SWidget> Owni
 			FCurveEditorScreenSpaceH HorizontalTransform = CurveEditor->GetPanelInputSpace();
 			double NewAnchorTime = HorizontalTransform.ScreenToSeconds(LocalPosition.X);
 
+			const FScopedTransaction Transaction(LOCTEXT("AddRetimingAnchorTransaction", "Add Retiming Anchor"));
+
+			RetimeData->Modify();
 			RetimeData->RetimingAnchors.Add(FCurveEditorRetimeAnchor(NewAnchorTime));
 
 			// Sort to ensure they're always sorted from least to greatest as drawing and between segments counts on this behavior.
