Curve Editor: Introduce Lattice tool

diff --git a/Engine/Plugins/Editor/CurveEditorTools/Source/CurveEditorTools/Private/CurveEditorToolCommands.cpp b/Engine/Plugins/Editor/CurveEditorTools/Source/CurveEditorTools/Private/CurveEditorToolCommands.cpp
--- a/Engine/Plugins/Editor/CurveEditorTools/Source/CurveEditorTools/Private/CurveEditorToolCommands.cpp
+++ b/Engine/Plugins/Editor/CurveEditorTools/Source/CurveEditorTools/Private/CurveEditorToolCommands.cpp
@@ -13,7 +13,8 @@ void FCurveEditorToolCommands::RegisterCommands()
 	// Tool Modes
 	UI_COMMAND(ActivateTransformTool, "Transform Tool", "Activates the Transform tool which allows translation, scale and rotation of selected keys.", EUserInterfaceActionType::ToggleButton, FInputChord(EModifierKey::Control, EKeys::T));
 	UI_COMMAND(ActivateRetimeTool, "Retime Tool", "Activates the Retime tool which allows you to define a one dimensional lattice to non-uniformly rescale key times.", EUserInterfaceActionType::ToggleButton, FInputChord(EModifierKey::Control, EKeys::E));
-	UI_COMMAND(ActivateMultiScaleTool, "Multi Scale Tool", "Activates the Multi Select tool which allows fast scaling of multiple curves from unique pivots", EUserInterfaceActionType::ToggleButton, FInputChord(EModifierKey::Control, EKeys::M));
+	UI_COMMAND(ActivateMultiScaleTool, "Multi Scale Tool", "Activates the Multi Select tool which allows fast scaling of multiple curves from unique pivots.", EUserInterfaceActionType::ToggleButton, FInputChord(EModifierKey::Control, EKeys::M));
+	UI_COMMAND(ActivateLatticeTool, "Lattice Tool", "Activates the Lattice tool which allows fast editing of keys using a grid.", EUserInterfaceActionType::ToggleButton, FInputChord(EModifierKey::Control, EKeys::L));
 }
 
 #undef LOCTEXT_NAMESPACE
diff --git a/Engine/Plugins/Editor/CurveEditorTools/Source/CurveEditorTools/Private/CurveEditorToolCommands.h b/Engine/Plugins/Editor/CurveEditorTools/Source/CurveEditorTools/Private/CurveEditorToolCommands.h
--- a/Engine/Plugins/Editor/CurveEditorTools/Source/CurveEditorTools/Private/CurveEditorToolCommands.h
+++ b/Engine/Plugins/Editor/CurveEditorTools/Source/CurveEditorTools/Private/CurveEditorToolCommands.h
@@ -27,6 +27,7 @@ public:
 	TSharedPtr<FUICommandInfo> ActivateTransformTool;
 	TSharedPtr<FUICommandInfo> ActivateRetimeTool;
 	TSharedPtr<FUICommandInfo> ActivateMultiScaleTool;
+	TSharedPtr<FUICommandInfo> ActivateLatticeTool;
 
 public:
 	virtual void RegisterCommands() override;
diff --git a/Engine/Plugins/Editor/CurveEditorTools/Source/CurveEditorTools/Private/CurveEditorToolsModule.cpp b/Engine/Plugins/Editor/CurveEditorTools/Source/CurveEditorTools/Private/CurveEditorToolsModule.cpp
--- a/Engine/Plugins/Editor/CurveEditorTools/Source/CurveEditorTools/Private/CurveEditorToolsModule.cpp
+++ b/Engine/Plugins/Editor/CurveEditorTools/Source/CurveEditorTools/Private/CurveEditorToolsModule.cpp
@@ -9,6 +9,7 @@
 #include "Tools/CurveEditorRetimeTool.h"
 #include "Tools/CurveEditorMultiScaleTool.h"
 #include "Framework/MultiBox/MultiBoxBuilder.h"
+#include "Tools/Lattice/CurveEditorLatticeTool.h"
 
 #define LOCTEXT_NAMESPACE "CurveEditorToolsModule"
 
@@ -24,6 +25,7 @@ private:
 	static TUniquePtr<ICurveEditorToolExtension> CreateTransformToolExtension(TWeakPtr<FCurveEditor> InCurveEditor);
 	static TUniquePtr<ICurveEditorToolExtension> CreateRetimeToolExtension(TWeakPtr<FCurveEditor> InCurveEditor);
 	static TUniquePtr<ICurveEditorToolExtension> CreateMultiScaleToolExtension(TWeakPtr<FCurveEditor> InCurveEditor);
+	static TUniquePtr<ICurveEditorToolExtension> CreateLatticeToolExtensions(TWeakPtr<FCurveEditor> InCurveEditor);
 
 protected:
 	TSharedRef<FExtender> ExtendCurveEditorToolbarMenu(const TSharedRef<FUICommandList> CommandList);
@@ -34,6 +36,7 @@ private:
 	FDelegateHandle TransformToolHandle;
 	FDelegateHandle RetimeToolHandle;
 	FDelegateHandle MultiScaleToolHandle;
+	FDelegateHandle LatticeToolHandle;
 };
 
 IMPLEMENT_MODULE(FCurveEditorToolsModule, CurveEditorTools)
@@ -51,6 +54,8 @@ void FCurveEditorToolsModule::StartupModule()
 	TransformToolHandle = CurveEditorModule.RegisterToolExtension(FOnCreateCurveEditorToolExtension::CreateStatic(&FCurveEditorToolsModule::CreateTransformToolExtension));
 	RetimeToolHandle = CurveEditorModule.RegisterToolExtension(FOnCreateCurveEditorToolExtension::CreateStatic(&FCurveEditorToolsModule::CreateRetimeToolExtension));
 	MultiScaleToolHandle = CurveEditorModule.RegisterToolExtension(FOnCreateCurveEditorToolExtension::CreateStatic(&FCurveEditorToolsModule::CreateMultiScaleToolExtension));
+	LatticeToolHandle = CurveEditorModule.RegisterToolExtension(FOnCreateCurveEditorToolExtension::CreateStatic(&FCurveEditorToolsModule::CreateLatticeToolExtensions));
+	
 
 	auto ToolbarExtender = ICurveEditorModule::FCurveEditorMenuExtender::CreateRaw(this, &FCurveEditorToolsModule::ExtendCurveEditorToolbarMenu);
 	auto& MenuExtenders = CurveEditorModule.GetAllToolBarMenuExtenders();
@@ -66,6 +71,7 @@ TSharedRef<FExtender> FCurveEditorToolsModule::ExtendCurveEditorToolbarMenu(cons
 			ToolbarBuilder.AddToolBarButton(FCurveEditorToolCommands::Get().ActivateTransformTool);
 			ToolbarBuilder.AddToolBarButton(FCurveEditorToolCommands::Get().ActivateRetimeTool);
 			ToolbarBuilder.AddToolBarButton(FCurveEditorToolCommands::Get().ActivateMultiScaleTool);
+			ToolbarBuilder.AddToolBarButton(FCurveEditorToolCommands::Get().ActivateLatticeTool);
 		}
 
 		static void FillToolbarFraming(FToolBarBuilder& ToolbarBuilder)
@@ -98,7 +104,8 @@ void FCurveEditorToolsModule::ShutdownModule()
 	// Unregister Tool Extensions
 	CurveEditorModule.UnregisterToolExtension(TransformToolHandle);
 	CurveEditorModule.UnregisterToolExtension(RetimeToolHandle);
-	CurveEditorModule.UnregisterEditorExtension(MultiScaleToolHandle);
+	CurveEditorModule.UnregisterToolExtension(MultiScaleToolHandle);
+	CurveEditorModule.UnregisterToolExtension(LatticeToolHandle);
 
 	FCurveEditorToolCommands::Unregister();
 }
@@ -123,4 +130,9 @@ TUniquePtr<ICurveEditorToolExtension> FCurveEditorToolsModule::CreateMultiScaleT
 	return MakeUnique<FCurveEditorMultiScaleTool>(InCurveEditor);
 }
 
+TUniquePtr<ICurveEditorToolExtension> FCurveEditorToolsModule::CreateLatticeToolExtensions(TWeakPtr<FCurveEditor> InCurveEditor)
+{
+	return MakeUnique<UE::CurveEditorTools::FCurveEditorLatticeTool>(InCurveEditor);
+}
+
 #undef LOCTEXT_NAMESPACE
diff --git a/Engine/Plugins/Editor/CurveEditorTools/Source/CurveEditorTools/Private/Misc/CurveChangeListener.cpp b/Engine/Plugins/Editor/CurveEditorTools/Source/CurveEditorTools/Private/Misc/CurveChangeListener.cpp
new file mode 100644
--- /dev/null
+++ b/Engine/Plugins/Editor/CurveEditorTools/Source/CurveEditorTools/Private/Misc/CurveChangeListener.cpp
@@ -0,0 +1,45 @@
+﻿// Copyright Epic Games, Inc. All Rights Reserved.
+
+#include "CurveChangeListener.h"
+
+#include "CurveEditor.h"
+#include "CurveModel.h"
+
+namespace UE::CurveEditorTools
+{
+	FCurveChangeListener::FCurveChangeListener(const TSharedRef<FCurveEditor>& InCurveEditor)
+		: FCurveChangeListener(InCurveEditor, [&InCurveEditor]()
+		{
+			TArray<FCurveModelID> Curves;
+			InCurveEditor->Selection.GetAll().GetKeys(Curves);
+			return Curves;
+		}())
+	{}
+
+	FCurveChangeListener::FCurveChangeListener(const TSharedRef<FCurveEditor>& InCurveEditor, const TArray<FCurveModelID>& InCurvesToListenTo)
+		: WeakCurveEditor(InCurveEditor)
+		, SubscribedToCurves(InCurvesToListenTo)
+	{
+		for (const FCurveModelID& CurveId : SubscribedToCurves)
+		{
+			if (FCurveModel* Model = InCurveEditor->FindCurve(CurveId))
+			{
+				Model->OnCurveModified().AddRaw(this, &FCurveChangeListener::HandleCurveModified);
+			}
+		}
+	}
+
+	FCurveChangeListener::~FCurveChangeListener()
+	{
+		if (TSharedPtr<FCurveEditor> CurveEditorPin = WeakCurveEditor.Pin())
+		{
+			for (const FCurveModelID& CurveId : SubscribedToCurves)
+			{
+				if (FCurveModel* Model = CurveEditorPin->FindCurve(CurveId))
+				{
+					Model->OnCurveModified().RemoveAll(this);
+				}
+			}
+		}
+	}
+}
diff --git a/Engine/Plugins/Editor/CurveEditorTools/Source/CurveEditorTools/Private/Misc/CurveChangeListener.h b/Engine/Plugins/Editor/CurveEditorTools/Source/CurveEditorTools/Private/Misc/CurveChangeListener.h
new file mode 100644
--- /dev/null
+++ b/Engine/Plugins/Editor/CurveEditorTools/Source/CurveEditorTools/Private/Misc/CurveChangeListener.h
@@ -0,0 +1,40 @@
+﻿// Copyright Epic Games, Inc. All Rights Reserved.
+
+#pragma once
+
+#include "Containers/Array.h"
+#include "Delegates/Delegate.h"
+#include "Templates/SharedPointer.h"
+#include "Templates/UnrealTemplate.h"
+
+class FCurveEditor;
+struct FCurveModelID;
+
+namespace UE::CurveEditorTools
+{
+/** Listens to changes made to curves and cleans the subscriptions on destruction. */
+class FCurveChangeListener : public FNoncopyable
+{
+public:
+	
+	explicit FCurveChangeListener(const TSharedRef<FCurveEditor>& InCurveEditor);
+	explicit FCurveChangeListener(const TSharedRef<FCurveEditor>& InCurveEditor, const TArray<FCurveModelID>& InCurvesToListenTo);
+	~FCurveChangeListener();
+
+	FSimpleMulticastDelegate& OnCurveModified() { return OnCurveModifiedDelegate; }
+
+private:
+
+	/** Used to unsubscribe on destruction. */
+	TWeakPtr<FCurveEditor> WeakCurveEditor;
+	/** The curves we subscribed to */
+	const TArray<FCurveModelID> SubscribedToCurves;
+
+	/** Invokes when any of the listeners is changed. */
+	FSimpleMulticastDelegate OnCurveModifiedDelegate;
+	
+	void HandleCurveModified() const { OnCurveModifiedDelegate.Broadcast(); }
+};
+}
+
+
diff --git a/Engine/Plugins/Editor/CurveEditorTools/Source/CurveEditorTools/Private/Misc/CurvePointSnapper.cpp b/Engine/Plugins/Editor/CurveEditorTools/Source/CurveEditorTools/Private/Misc/CurvePointSnapper.cpp
new file mode 100644
--- /dev/null
+++ b/Engine/Plugins/Editor/CurveEditorTools/Source/CurveEditorTools/Private/Misc/CurvePointSnapper.cpp
@@ -0,0 +1,44 @@
+﻿// Copyright Epic Games, Inc. All Rights Reserved.
+
+#include "CurvePointSnapper.h"
+
+#include "CurveEditor.h"
+#include "CurveEditorTypes.h"
+#include "SCurveEditorView.h"
+#include "Misc/Optional.h"
+
+namespace UE::CurveEditorTools
+{
+TOptional<FCurvePointSnapper> FCurvePointSnapper::MakeSnapper(FCurveEditor& InCurveEditor)
+{
+	if (InCurveEditor.GetSelection().GetAll().IsEmpty())
+	{
+		return {};
+	}
+	
+	// Snap the mouse to the grid according to the first curve's snap metrics (assuming all curves have the same view scales)
+	const FCurveModelID FirstCurveID = InCurveEditor.GetSelection().GetAll().CreateConstIterator()->Key;
+	const SCurveEditorView* View = InCurveEditor.FindFirstInteractiveView(FirstCurveID);
+	if (!View)
+	{
+		return {};
+	}
+
+	return FCurvePointSnapper(View->GetCurveSpace(FirstCurveID), InCurveEditor.GetCurveSnapMetrics(FirstCurveID));
+}
+
+FVector2D FCurvePointSnapper::SnapPoint(const FVector2D& InCurveSpacePoint) const
+{
+	return FVector2D
+	{
+		SnapMetrics.SnapInputSeconds(InCurveSpacePoint.X),
+		SnapMetrics.SnapOutput(InCurveSpacePoint.Y)
+	};
+}
+
+FKeyPosition FCurvePointSnapper::SnapKey(const FKeyPosition& InKeyPosition) const
+{
+	const FVector2D Point = SnapPoint({ InKeyPosition.InputValue, InKeyPosition.OutputValue });
+	return { Point.X, Point.Y };
+}
+}
diff --git a/Engine/Plugins/Editor/CurveEditorTools/Source/CurveEditorTools/Private/Misc/CurvePointSnapper.h b/Engine/Plugins/Editor/CurveEditorTools/Source/CurveEditorTools/Private/Misc/CurvePointSnapper.h
new file mode 100644
--- /dev/null
+++ b/Engine/Plugins/Editor/CurveEditorTools/Source/CurveEditorTools/Private/Misc/CurvePointSnapper.h
@@ -0,0 +1,38 @@
+﻿// Copyright Epic Games, Inc. All Rights Reserved.
+
+#pragma once
+
+#include "CurveEditorScreenSpace.h"
+#include "CurveEditorSnapMetrics.h"
+
+struct FKeyPosition;
+class FCurveEditor;
+template<typename OptionalType> struct TOptional;
+
+namespace UE::CurveEditorTools
+{
+/** Util that you can snap points according to a curve editor's snapping settings. */
+class FCurvePointSnapper
+{
+public:
+
+	/** @return Unset if there's no curve view from which to get the FCurveEditorScreenSpace transform. */
+	static TOptional<FCurvePointSnapper> MakeSnapper(FCurveEditor& InCurveEditor);
+
+	/** Snaps the point as per the snapping settings */
+	FVector2D SnapPoint(const FVector2D& InCurveSpacePoint) const;
+
+	/** Snaps the key as per the snapping settings */
+	FKeyPosition SnapKey(const FKeyPosition& InKeyPosition) const;
+
+private:
+
+	explicit FCurvePointSnapper(const FCurveEditorScreenSpace& InCurveSpace, const FCurveSnapMetrics& InMetrics)
+		: CurveSpace(InCurveSpace), SnapMetrics(InMetrics)
+	{}
+
+	const FCurveEditorScreenSpace CurveSpace;
+	const FCurveSnapMetrics SnapMetrics;
+};
+}
+
diff --git a/Engine/Plugins/Editor/CurveEditorTools/Source/CurveEditorTools/Private/Misc/VectorMathUtils.cpp b/Engine/Plugins/Editor/CurveEditorTools/Source/CurveEditorTools/Private/Misc/VectorMathUtils.cpp
new file mode 100644
--- /dev/null
+++ b/Engine/Plugins/Editor/CurveEditorTools/Source/CurveEditorTools/Private/Misc/VectorMathUtils.cpp
@@ -0,0 +1,51 @@
+﻿// Copyright Epic Games, Inc. All Rights Reserved.
+
+#include "VectorMathUtils.h"
+
+#include "Math/Vector2D.h"
+
+namespace UE::CurveEditorTools
+{
+static double CrossProductSign(const FVector2D& Point, const FVector2D& Vertex1, const FVector2D& Vertex2)
+{
+	const FVector2D Vertex1ToPoint = Point - Vertex1;
+	const FVector2D Vertex1ToVertex2 = Vertex2 - Vertex1;
+	return FVector2D::CrossProduct(Vertex1ToPoint, Vertex1ToVertex2);
+}
+
+bool IsPointInTriangle(const FVector2D& Point, const FVector2D& Vertex1, const FVector2D& Vertex2, const FVector2D& Vertex3)
+{
+	// The idea is to compute the 2D cross product between the point and each edge.
+	// Point x edge < 0 means point is to the right
+	// Point x edge > 0 means point is to the left
+	const double Sign_1_2 = CrossProductSign(Point, Vertex1, Vertex2);
+	const double Sign_2_3 = CrossProductSign(Point, Vertex2, Vertex3);
+	const double Sign_3_1 = CrossProductSign(Point, Vertex3, Vertex1);
+
+	// If the point is on the same side of each edge, that means it is inside the triangle.
+	const bool bHasPositive = Sign_1_2 > 0.0 || Sign_2_3 > 0.0 || Sign_3_1 > 0.0;
+	const bool bHasNegative = Sign_1_2 < 0.0 || Sign_2_3 < 0.0 || Sign_3_1 < 0.0;
+	return !(bHasPositive && bHasNegative);
+}
+
+void InsetQuadBy(FVector2D& A, FVector2D& B, FVector2D& C, FVector2D& D, float InsetAmount)
+{
+	// Calculate edge vectors and normals
+	const FVector2D AB = (B - A).GetSafeNormal(); 
+	const FVector2D BC = (C - B).GetSafeNormal(); 
+	const FVector2D CD = (D - C).GetSafeNormal(); 
+	const FVector2D DA = (A - D).GetSafeNormal(); 
+
+	// Calculate inward normals (perpendiculars)
+	const FVector2D InwardAB = FVector2D(-AB.Y, AB.X); 
+	const FVector2D InwardBC = FVector2D(-BC.Y, BC.X); 
+	const FVector2D InwardCD = FVector2D(-CD.Y, CD.X); 
+	const FVector2D InwardDA = FVector2D(-DA.Y, DA.X); 
+
+	// Offset each vertex by the inward normals of adjacent edges
+	A = A + -InsetAmount * (InwardDA + InwardAB).GetSafeNormal();
+	B = B + -InsetAmount * (InwardAB + InwardBC).GetSafeNormal();
+	C = C + -InsetAmount * (InwardBC + InwardCD).GetSafeNormal();
+	D = D + -InsetAmount * (InwardCD + InwardDA).GetSafeNormal();
+}
+}
diff --git a/Engine/Plugins/Editor/CurveEditorTools/Source/CurveEditorTools/Private/Misc/VectorMathUtils.h b/Engine/Plugins/Editor/CurveEditorTools/Source/CurveEditorTools/Private/Misc/VectorMathUtils.h
new file mode 100644
--- /dev/null
+++ b/Engine/Plugins/Editor/CurveEditorTools/Source/CurveEditorTools/Private/Misc/VectorMathUtils.h
@@ -0,0 +1,13 @@
+﻿// Copyright Epic Games, Inc. All Rights Reserved.
+
+#pragma once
+
+#include "Math/MathFwd.h"
+
+namespace UE::CurveEditorTools
+{
+	/** @return Whether Point is in the triangle spanned by the input vertices. */
+	bool IsPointInTriangle(const FVector2D& Point, const FVector2D& Vertex1, const FVector2D& Vertex2, const FVector2D& Vertex3);
+	/** @return Makes a smaller quad in the same shape but with the edges moved inside by InsetAmount. */
+	void InsetQuadBy(FVector2D& A, FVector2D& B, FVector2D& C, FVector2D& D, float InsetAmount);
+}
diff --git a/Engine/Plugins/Editor/CurveEditorTools/Source/CurveEditorTools/Private/Tools/Lattice/CurveEditorLatticeTool.cpp b/Engine/Plugins/Editor/CurveEditorTools/Source/CurveEditorTools/Private/Tools/Lattice/CurveEditorLatticeTool.cpp
new file mode 100644
--- /dev/null
+++ b/Engine/Plugins/Editor/CurveEditorTools/Source/CurveEditorTools/Private/Tools/Lattice/CurveEditorLatticeTool.cpp
@@ -0,0 +1,475 @@
+﻿// Copyright Epic Games, Inc. All Rights Reserved.
+
+#include "CurveEditorLatticeTool.h"
+
+#include "CurveEditor.h"
+#include "CurveEditorToolCommands.h"
+#include "Drag/LatticeDragOp_MoveControlPoints.h"
+#include "Internationalization/Text.h"
+#include "Misc/CurvePointSnapper.h"
+#include "Misc/LatticeDrawUtils.h"
+#include "Misc/LatticeUtils.h"
+#include "Templates/UnrealTemplate.h"
+#include "UObject/LatticeUndoData.h"
+#include "UObject/Package.h"
+#include "UObject/UObjectGlobals.h"
+
+#define LOCTEXT_NAMESPACE "FCurveEditorLatticeTool"
+
+namespace UE::CurveEditorTools
+{
+FCurveEditorLatticeTool::FCurveEditorLatticeTool(TWeakPtr<FCurveEditor> InCurveEditor)
+	: WeakCurveEditor(InCurveEditor)
+	, FinalMouseScreenSpacePosition()
+	, UndoObject(
+		NewObject<UCurveEditorTools_LatticeUndoData>(GetTransientPackage(), UCurveEditorTools_LatticeUndoData::StaticClass(), NAME_None, RF_Transactional)
+	)
+{
+	UndoObject->OnPostEditUndo.AddLambda([this, WeakCurveEditor = WeakCurveEditor]
+	{
+		const TSharedPtr<FCurveEditor> Editor = WeakCurveEditor.Pin();
+		if (!Editor)
+		{
+			return;
+		}
+
+		// LatticeDeformerState is unset if the user redo-es the first transaction which set the initial selection...
+		if (LatticeDeformerState)
+		{
+			UndoObject->RestoreDeformerState(GetDeformer());
+			
+			UndoObject->RestoreEditorState(*Editor);
+			// We need to overwrite our serial number of next frame UpdateDeformerBounds will think that the selection has changed and reinit.
+			SelectionSerialNumber = UndoObject->GetSelectionSerialNumber();
+
+			// If user undo-es the first transaction, next tick we'll notice that nothing is selected and destroy the deformer.
+		}
+		else
+		{
+			// ... in that case, we refresh the selection and next tick, we'll compute the bounds, which will end up have the rectangular starting shape.
+			UndoObject->RestoreEditorState(*Editor);
+		}
+	});
+}
+
+void FCurveEditorLatticeTool::Tick(const FGeometry& InViewContainerGeometry, const double InCurrentTime, const float InDeltaTime)
+{
+	// DrawData needs to be updated in Tick. InViewContainerGeometry is slightly smaller in OnPaint than in OnTick/OnMouseButtonDown, etc.
+	// This results in the hover state not being displayed correctly. For that reason, we generate the data in OnTick and draw it in OnPaint.
+	const TSharedPtr<FCurveEditor> CurveEditor = WeakCurveEditor.Pin();
+	if (!CurveEditor)
+	{
+		return;
+	}
+	
+	if (LatticeDeformerState && LatticeDeformerState->bHasRequestedRefresh)
+	{
+		LatticeDeformerState->bHasRequestedRefresh = false;
+		AdjustLatticeGridOrRecompute(CurveEditor.ToSharedRef());
+	}
+	else
+	{
+		UpdateDeformerBounds(CurveEditor.ToSharedRef());
+	}
+		
+	if (LatticeDeformerState.IsSet())
+	{
+		DrawData.Emplace();
+		RecomputeDrawData(InViewContainerGeometry, CurveEditor);
+	}
+	else
+	{
+		DrawData.Reset();
+	}
+}
+
+void FCurveEditorLatticeTool::OnPaint(
+	const FPaintArgs& InPaintArgs,
+	const FGeometry& InViewContainerGeometry,
+	const FSlateRect& InMyCullingRect,
+	FSlateWindowElementList& OutDrawElements,
+	int32 InPaintOnLayerId,
+	const FWidgetStyle& InWidgetStyle,
+	bool bParentEnabled
+	) const
+{
+	if (DrawData.IsSet())
+	{
+		DrawLatticeControls(*DrawData, InViewContainerGeometry, OutDrawElements, InPaintOnLayerId);
+	}
+}
+
+FReply FCurveEditorLatticeTool::OnMouseButtonDown(
+	TSharedRef<SWidget> InOwningWidget, const FGeometry& InViewContainerGeometry, const FPointerEvent& InMouseEvent
+	)
+{
+	CancelDrag();
+	
+	if (const TSharedPtr<FCurveEditor> CurveEditor = WeakCurveEditor.Pin()
+		; CurveEditor && LatticeDeformerState.IsSet())
+	{
+		const FLatticeBounds Bounds = ComputeBounds(*CurveEditor);
+		if (Bounds.bIsVisible && InMouseEvent.GetEffectingButton() == EKeys::LeftMouseButton)
+		{
+			const TArray<FVector2D> ControlPointsInSlate = ComputeDrawControlPoints(*CurveEditor);
+			const FLatticeHoverState HoverState = ComputeLatticeHoverState(
+				InViewContainerGeometry, InMouseEvent.GetScreenSpacePosition(),
+				ControlPointsInSlate, ComputeDrawControlEdges(ControlPointsInSlate),
+				GetDeformer().NumPointsInWidth(), GetDeformer().NumCells()
+				);
+			if (HoverState.IsHovered())
+			{
+				DelayedDragData.Emplace(
+					InMouseEvent.GetScreenSpacePosition(),
+					InMouseEvent.GetEffectingButton(),
+					Bounds,
+					HoverState
+					);
+				return FReply::Handled().CaptureMouse(InOwningWidget).PreventThrottling();
+			}
+		}
+	}
+	
+	return FReply::Unhandled();
+}
+
+FReply FCurveEditorLatticeTool::OnMouseButtonUp(
+	TSharedRef<SWidget> InOwningWidget, const FGeometry& InViewContainerGeometry, const FPointerEvent& MouseEvent
+	)
+{
+	FReply Reply = FReply::Handled();
+	if (DelayedDragData.IsSet())
+	{
+		if (DelayedDragData->DelayedDrag.IsDragging() && ensure(DelayedDragData->ActiveOperation))
+		{
+			DelayedDragData->ActiveOperation->DragOp->EndDrag(InViewContainerGeometry, MouseEvent);
+			// Only return handled if we actually started a drag
+			Reply.ReleaseMouseCapture();
+		}
+		DelayedDragData.Reset();
+		return Reply;
+	}
+	
+	return FReply::Unhandled();
+}
+
+FReply FCurveEditorLatticeTool::OnMouseMove(
+	TSharedRef<SWidget> InOwningWidget, const FGeometry& InViewContainerGeometry, const FPointerEvent& InMouseEvent
+	)
+{
+	FinalMouseScreenSpacePosition = InMouseEvent.GetScreenSpacePosition();
+	
+	if (DelayedDragData.IsSet())
+	{
+		FReply Reply = FReply::Handled();
+
+		if (DelayedDragData->DelayedDrag.IsDragging() && ensure(DelayedDragData->ActiveOperation))
+		{
+			DelayedDragData->ActiveOperation->DragOp->MoveMouse(InViewContainerGeometry, InMouseEvent);
+		}
+		else if (DelayedDragData->DelayedDrag.AttemptDragStart(InMouseEvent))
+		{
+			DelayedDragData->ActiveOperation.Emplace(
+				CreateDragDrop(DelayedDragData->Bounds, DelayedDragData->HoverState),
+				LOCTEXT("LatticeChangeKeys", "Lattice keys")
+				);
+			DelayedDragData->ActiveOperation->DragOp->BeginDrag(
+				InViewContainerGeometry, InMouseEvent, DelayedDragData->DelayedDrag.GetInitialPosition()
+				);
+		}
+			
+		// Steal the capture, as we're now the authoritative widget in charge of a mouse-drag operation
+		Reply.CaptureMouse(InOwningWidget);	
+		return Reply;
+	}
+	
+	return ICurveEditorToolExtension::OnMouseMove(InOwningWidget, InViewContainerGeometry, InMouseEvent);
+}
+
+void FCurveEditorLatticeTool::OnFinishedPointerInput()
+{
+	if (DelayedDragData && DelayedDragData->ActiveOperation)
+	{
+		DelayedDragData->ActiveOperation->DragOp->FinishedPointerInput();
+	}
+}
+
+void FCurveEditorLatticeTool::OnToolActivated()
+{
+	if (const TSharedPtr<FCurveEditor> CurveEditor = WeakCurveEditor.Pin())
+	{
+		UpdateDeformerBounds(CurveEditor.ToSharedRef());
+	}
+}
+
+void FCurveEditorLatticeTool::OnToolDeactivated()
+{
+	// Don't clear the data - if the user activates after deactivating without changing the selected keys, we'll show them what they last had
+}
+
+void FCurveEditorLatticeTool::BindCommands(TSharedRef<FUICommandList> CommandBindings)
+{
+	TSharedPtr<FCurveEditor> CurveEditor = WeakCurveEditor.Pin();
+	if (CurveEditor)
+	{
+		FIsActionChecked MultiScaleToolIsActive = FIsActionChecked::CreateSP(CurveEditor.ToSharedRef(), &FCurveEditor::IsToolActive, ToolID);
+		FExecuteAction ActivateMultiScaleTool = FExecuteAction::CreateSP(CurveEditor.ToSharedRef(), &FCurveEditor::MakeToolActive, ToolID);
+
+		CommandBindings->MapAction(
+			FCurveEditorToolCommands::Get().ActivateLatticeTool, ActivateMultiScaleTool, FCanExecuteAction(), MultiScaleToolIsActive
+			);
+	}
+}
+
+void FCurveEditorLatticeTool::AddReferencedObjects(FReferenceCollector& Collector)
+{
+	Collector.AddReferencedObject(UndoObject);
+}
+
+FCurveEditorLatticeTool::FLatticeDeformerState::FLatticeDeformerState(const FLatticeBounds& Lattice, const TSharedRef<FCurveEditor>& InCurveEditor)
+	: Deformer(1, 1, Lattice.MinValues, Lattice.MaxValues)
+	, CurveChangeListener(InCurveEditor)
+{}
+
+void FCurveEditorLatticeTool::UpdateDeformerBounds(const TSharedRef<FCurveEditor>& InCurveEditor)
+{
+	if (SelectionSerialNumber != InCurveEditor->Selection.GetSerialNumber())
+	{
+		CancelDrag();
+		LatticeDeformerState.Reset();
+	}
+
+	if (LatticeDeformerState.IsSet())
+	{
+		return;
+	}
+
+	// Handle the edge case of all points overlapping, or a single key selection.
+	// The lattice math does not work when widht or height are 0 due to division by 0; fyi, the deformer handles it with an ensure.
+	const FLatticeBounds Bounds = ComputeBounds(*InCurveEditor);
+	const FVector2D DeltaValues = Bounds.MaxValues - Bounds.MinValues;
+	const bool bTooSmall = DeltaValues.X < UE_KINDA_SMALL_NUMBER || DeltaValues.Y < UE_KINDA_SMALL_NUMBER || InCurveEditor->Selection.Count() <= 1;
+	if (Bounds.bIsVisible && !bTooSmall)
+	{
+		SelectionSerialNumber = InCurveEditor->Selection.GetSerialNumber();
+		LatticeDeformerState.Emplace(Bounds, InCurveEditor);
+		AddSelectionToDeformer(GetDeformer(), *InCurveEditor);
+		LatticeDeformerState->CurveChangeListener.OnCurveModified().AddRaw(this, &FCurveEditorLatticeTool::OnCurvesModified);
+
+		FScopedTransaction Transaction(TEXT("CurveEditorLatticeTool"), LOCTEXT("SetSelection", "Set lattice selection"), nullptr);
+		UndoObject->Modify();
+		UndoObject->SaveState(*InCurveEditor, GetDeformer());
+	}
+}
+
+void FCurveEditorLatticeTool::OnCurvesModified()
+{
+	// The modification was driven by us?
+	LatticeDeformerState->bHasRequestedRefresh = !LatticeDeformerState->bIsModifyingCurves;
+}
+
+void FCurveEditorLatticeTool::AdjustLatticeGridOrRecompute(const TSharedRef<FCurveEditor>& InCurveEditor)
+{
+	checkSlow(LatticeDeformerState);
+	
+	// If the selection has changed completely, then refitting the pre-existing lattice makes no sense.
+	const bool bHasSelectionChanged = SelectionSerialNumber != InCurveEditor->Selection.GetSerialNumber();
+	if (bHasSelectionChanged)
+	{
+		// This will calculate a new rectangular grid around the points.
+		// Instead, we could try to preserve the existing shape by just moving it over the moved keys, provided it still fits.
+		LatticeDeformerState.Reset();
+		UpdateDeformerBounds(InCurveEditor);
+	}
+	else
+	{
+		if (GUndo)
+		{
+			UndoObject->Modify();
+			UndoObject->SaveState(*InCurveEditor, GetDeformer());
+		}
+		
+		const TArray<FVector2D> TranslatedControlPoints = PlaceLatticeGridOverNewKeyPositions(GetDeformer(), *InCurveEditor);
+		checkSlow(TranslatedControlPoints.Num() == GetDeformer().GetControlPoints().Num());
+		for (int32 Index = 0; Index < TranslatedControlPoints.Num(); ++Index)
+		{
+			const FVector2D& UpdatedControlPoint = TranslatedControlPoints[Index];
+			const TConstArrayView<FVector2D> DummyView(&UpdatedControlPoint, 1);
+			GetDeformer().SetControlsPointNoRecompute({ Index }, DummyView);
+		}
+	}
+}
+
+void FCurveEditorLatticeTool::CancelDrag()
+{
+	if (DelayedDragData && DelayedDragData->ActiveOperation)
+	{
+		DelayedDragData->ActiveOperation->DragOp->CancelDrag();
+	}
+	DelayedDragData.Reset();
+}
+	
+void FCurveEditorLatticeTool::RecomputeDrawData(const FGeometry& InViewContainerGeometry, const TSharedPtr<FCurveEditor>& InCurveEditor)
+{
+	// The widget hierarchy is SCurveEditorViewContainer -> SCurveEditorViewAbsolute.
+	// InAllottedGeometry is that of SCurveEditorViewContainer.
+	DrawData->MatrixWidth = GetDeformer().NumPointsInWidth();
+	DrawData->ControlPoints = ComputeDrawControlPoints(*InCurveEditor);
+	const bool bHasControlData = !DrawData->ControlPoints.IsEmpty();
+	if (!bHasControlData)
+	{
+		return;
+	}
+	
+	DrawData->ControlEdges = ComputeDrawControlEdges(DrawData->ControlPoints);
+	// No hover highlights while dragging
+	if (!DelayedDragData.IsSet())
+	{
+		DrawData->HoverState = ComputeLatticeHoverState(
+			InViewContainerGeometry, FinalMouseScreenSpacePosition,
+			DrawData->ControlPoints,DrawData->ControlEdges,
+			GetDeformer().NumPointsInWidth(), GetDeformer().NumCells()
+		);
+	}
+}
+
+TArray<FVector2D> FCurveEditorLatticeTool::ComputeDrawControlPoints(const FCurveEditor& InCurveEditor) const
+{
+	check(LatticeDeformerState);
+	TArray<FVector2D> Result;
+
+	const TConstArrayView<FVector2D> ControlPoints = GetDeformer().GetControlPoints();
+	Result.SetNumUninitialized(ControlPoints.Num());
+	if (!TransformViewToViewContainer(InCurveEditor, ControlPoints, Result))
+	{
+		return {};
+	}
+	return Result;
+}
+
+TArray<FLatticeControlEdge> FCurveEditorLatticeTool::ComputeDrawControlEdges(TConstArrayView<FVector2D> InControlPoints) const
+{
+	TArray<FLatticeControlEdge> ControlEdges;
+	ControlEdges.SetNumUninitialized(GetDeformer().NumEdges());
+	for (int32 EdgeIndex = 0; EdgeIndex < GetDeformer().NumEdges(); ++EdgeIndex)
+	{
+		const auto[Start, End] = GetDeformer().GetEdgeIndices(EdgeIndex);
+		ControlEdges[EdgeIndex] = FLatticeControlEdge(InControlPoints[Start], InControlPoints[End]);
+	}
+	return ControlEdges;
+}
+	
+TUniquePtr<FLatticeDragOp> FCurveEditorLatticeTool::CreateDragDrop(const FLatticeBounds& InBounds, const FLatticeHoverState& InHoverState)
+{
+	const TConstArrayView<FVector2D> ControlPoints = GetDeformer().GetControlPoints();
+	if (InHoverState.HoveredControlPoint)
+	{
+		const int32 ControlPointIndex = *InHoverState.HoveredControlPoint;
+		return MakeUnique<FLatticeDragOp_MoveControlPoints>(
+			WeakCurveEditor,
+			FControlPointArray{ ControlPoints[ControlPointIndex] },
+			InBounds,
+			FLatticeDragOp_MoveControlPoints::FOnControlPointsMoved::CreateLambda(
+				[this, ControlPointIndex](const TConstArrayView<FVector2D>& InNewControlPointPositions)
+				{
+					OnSetControlPoints({ ControlPointIndex }, InNewControlPointPositions);
+				}
+			)
+		);
+	}
+
+	if (InHoverState.HoveredEdge)
+	{
+		const auto[Start, End] = GetDeformer().GetEdgeIndices(*InHoverState.HoveredEdge);
+		return MakeUnique<FLatticeDragOp_MoveControlPoints>(
+			WeakCurveEditor,
+			FControlPointArray{ ControlPoints[Start], ControlPoints[End] },
+			InBounds,
+			FLatticeDragOp_MoveControlPoints::FOnControlPointsMoved::CreateLambda(
+				[this, Start, End](const TConstArrayView<FVector2D>& InNewControlPointPositions)
+				{
+					OnSetControlPoints({ Start, End }, InNewControlPointPositions);
+				}
+			)
+		);
+	}
+
+	if (InHoverState.HoveredCell)
+	{
+		const auto[TopLeft, TopRight, BottomRight, BottomLeft] = GetDeformer().GetCellIndices(*InHoverState.HoveredCell);
+		return MakeUnique<FLatticeDragOp_MoveControlPoints>(
+			WeakCurveEditor,
+			FControlPointArray{ ControlPoints[TopLeft], ControlPoints[TopRight], ControlPoints[BottomRight], ControlPoints[BottomLeft] },
+			InBounds,
+			FLatticeDragOp_MoveControlPoints::FOnControlPointsMoved::CreateLambda(
+				[this, TopLeft, TopRight, BottomRight, BottomLeft](const TConstArrayView<FVector2D>& InNewControlPointPositions)
+				{
+					OnSetControlPoints({ TopLeft, TopRight, BottomRight, BottomLeft }, InNewControlPointPositions);
+				}
+			)
+		);
+	}
+
+	checkNoEntry();
+	return nullptr;
+}
+
+void FCurveEditorLatticeTool::OnSetControlPoints(TConstArrayView<int32> InControlPointsToUpdate, TConstArrayView<FVector2D> InNewControlPoints)
+{
+	const TSharedPtr<FCurveEditor> CurveEditor = WeakCurveEditor.Pin();
+	const TOptional<FCurvePointSnapper> Snapper = CurveEditor ? FCurvePointSnapper::MakeSnapper(*CurveEditor) : TOptional<FCurvePointSnapper>{};
+	if (!CurveEditor || !Snapper)
+	{
+		return;
+	}
+	
+	// Combining Handles & Keys into large arrays like this is significantly faster than feeding each key into SetKeyPositions separately.
+	// For 9000 keys it goes from 2136 ms to 238 ms.
+	struct FCurveChanges
+	{
+		TArray<FKeyHandle> Handles;
+		TArray<FKeyPosition> Keys;
+	};
+
+	TMap<FCurveModelID, FCurveChanges> Changes;
+	// Pre-allocating saves about 0.060ms
+	for (const TPair<FCurveModelID, FKeyHandleSet>& Pair : CurveEditor->Selection.GetAll())
+	{
+		FCurveChanges& CurveChanges = Changes.Add(Pair.Key);
+		CurveChanges.Handles.Reserve(Pair.Value.Num());
+		CurveChanges.Keys.Reserve(Pair.Value.Num());
+	}
+	
+	// The actual re-computation takes about 0.083 ms
+	GetDeformer().UpdateControlPoints(InControlPointsToUpdate, InNewControlPoints,
+		[&Snapper, &Changes](const FKeyInCurveId& Handle, const FVector2D& KeyPosition)
+		{
+			FCurveChanges& CurveChanges = Changes.FindOrAdd(Handle.CurveId);
+			CurveChanges.Handles.Add(Handle.KeyId);
+			CurveChanges.Keys.Add(
+				Snapper->SnapKey(FKeyPosition{ KeyPosition.X, KeyPosition.Y })
+				);
+		});
+
+	// FDragOpData has started a FScopedTransaction and will release it once dragging has stopped.
+	UndoObject->Modify();
+	UndoObject->SaveState(*CurveEditor, GetDeformer());
+
+	// Some stats all with 9000 keys. Varying degree of changed keys, i.e. how far you move your mouse in quick succession.
+	// Sorted by little to a lot of changes:
+	// - SetKeyPositions: 4.299 ms		- Total: 4.533 ms
+	// - SetKeyPositions 138.954 ms		- Total 139.229 ms
+	// - SetKeyPositions 269.523 ms		- Total 269.811 ms
+	// Conclusion: The following SetKeyPositions block is slow.
+	TGuardValue Guard(LatticeDeformerState->bIsModifyingCurves, true);
+	for (auto It = Changes.CreateConstIterator(); It; ++It)
+	{
+		const FCurveChanges& Pair = It->Value;
+		if (FCurveModel* Model = CurveEditor->FindCurve(It->Key); ensure(Model))
+		{
+			Model->Modify();
+			Model->SetKeyPositions(Pair.Handles, Pair.Keys);
+		}
+	}
+}
+}
\ No newline at end of file
diff --git a/Engine/Plugins/Editor/CurveEditorTools/Source/CurveEditorTools/Private/Tools/Lattice/CurveEditorLatticeTool.h b/Engine/Plugins/Editor/CurveEditorTools/Source/CurveEditorTools/Private/Tools/Lattice/CurveEditorLatticeTool.h
new file mode 100644
--- /dev/null
+++ b/Engine/Plugins/Editor/CurveEditorTools/Source/CurveEditorTools/Private/Tools/Lattice/CurveEditorLatticeTool.h
@@ -0,0 +1,148 @@
+﻿// Copyright Epic Games, Inc. All Rights Reserved.
+
+#pragma once
+
+#include "Drag/LatticeDragOp.h"
+#include "Framework/DelayedDrag.h"
+#include "ICurveEditorToolExtension.h"
+#include "LatticeControlsDrawData.h"
+#include "LatticeDeformer2D.h"
+#include "Misc/CurveChangeListener.h"
+#include "Misc/LatticeDrawUtils.h"
+#include "ScopedTransaction.h"
+#include "Templates/SharedPointer.h"
+
+class FCurveEditor;
+class FText;
+class UCurveEditorTools_LatticeUndoData;
+
+namespace UE::CurveEditorTools
+{
+struct FLatticeBounds;
+	
+/**
+ * Places a simple lattice deformer over the selection.
+ * The grid is a simple rectangle. You can reshape by moving the points or edges of the rectangle.
+ */
+class FCurveEditorLatticeTool : public ICurveEditorToolExtension, public FGCObject
+{
+public:
+	
+	explicit FCurveEditorLatticeTool(TWeakPtr<FCurveEditor> InCurveEditor);
+
+	//~ Begin ICurveEditorToolExtension Interface
+	virtual void Tick(const FGeometry& InViewContainerGeometry, const double InCurrentTime, const float InDeltaTime) override;
+	virtual void OnPaint(const FPaintArgs& InPaintArgs, const FGeometry& InViewContainerGeometry, const FSlateRect& InMyCullingRect, FSlateWindowElementList& OutDrawElements, int32 InPaintOnLayerId, const FWidgetStyle& InWidgetStyle, bool bParentEnabled) const override;
+	virtual FReply OnMouseButtonDown(TSharedRef<SWidget> InOwningWidget, const FGeometry& InViewContainerGeometry, const FPointerEvent& InMouseEvent) override;
+	virtual FReply OnMouseButtonUp(TSharedRef<SWidget> InOwningWidget, const FGeometry& InViewContainerGeometry, const FPointerEvent& MouseEvent) override;
+	virtual FReply OnMouseMove(TSharedRef<SWidget> InOwningWidget, const FGeometry& InViewContainerGeometry, const FPointerEvent& InMouseEvent) override;
+	virtual void OnFinishedPointerInput() override;
+	virtual void OnToolActivated() override;
+	virtual void OnToolDeactivated() override;
+	virtual void BindCommands(TSharedRef<FUICommandList> CommandBindings) override;
+	//~ End ICurveEditorToolExtension Interface
+
+	//~ Begin FGCObject Interface
+	virtual void AddReferencedObjects(FReferenceCollector& Collector) override;
+	virtual FString GetReferencerName() const override {  return FString("FCurveEditorLatticeTool"); }
+	//~ End FGCObject Interface
+
+private:
+	
+	/** Weak pointer back to the owning curve editor. */
+	TWeakPtr<FCurveEditor> WeakCurveEditor;
+
+	/** Holds any data that exists only when the deformer exists. The deformer exists the selection consists of least 2 keys at different locations. */
+	struct FLatticeDeformerState
+	{
+		/** Does the actual deforming calculations of the points. */
+		FLatticeDeformer2D Deformer;
+
+		/**
+		 * Tells us when a curve is externally modified, so we can update the bounds.
+		 * Listens for changes on individual curves and cleans up the subscriptions on destruction.
+		 */
+		FCurveChangeListener CurveChangeListener;
+		
+		/** Guard to not recompute bounds when FCurveModel::OnCurveModified broadcasts due to a change we have initiated. */
+		bool bIsModifyingCurves = false;
+		/** If true, the curves have been modified and the lattice overlay should be recomputed next frame. */
+		bool bHasRequestedRefresh = false;
+
+		explicit FLatticeDeformerState(const FLatticeBounds& Lattice, const TSharedRef<FCurveEditor>& InCurveEditor);
+	};
+	/** Set when there is a valid selection in view. */
+	TOptional<FLatticeDeformerState> LatticeDeformerState;
+	/** The selection serial number when the deformer was initialized. */
+	uint32 SelectionSerialNumber = 0;
+	
+	/** How to draw LatticeDeformer. Set only when the selection, over which the controls are drawn, is visible.*/
+	TOptional<FLatticeControlsDrawData> DrawData;
+
+	/** The last mouse position received by OnMouseMove. */
+	FVector2D FinalMouseScreenSpacePosition;
+
+	struct FDragOpData
+	{
+		/** Handles drags that started on a control point. */
+		TUniquePtr<FLatticeDragOp> DragOp;
+		FScopedTransaction Transaction;
+
+		FDragOpData(TUniquePtr<FLatticeDragOp> DragOp, const FText& TransactionText)
+			: DragOp(MoveTemp(DragOp))
+			, Transaction(TEXT("CurveEditorLatticeTool"), TransactionText, nullptr)
+		{}
+	};
+	
+	struct FDelayedDragData
+	{
+		/** Set when attempting to move a drag handle. This allows us to tell the difference between a click and a click-drag. */
+		FDelayedDrag DelayedDrag;
+
+		/** Bounds when the op was started. */
+		const FLatticeBounds Bounds;
+		/** Hover state when the op was started. */
+		const FLatticeHoverState HoverState;
+
+		/** Set if the drag up was started. */
+		TOptional<FDragOpData> ActiveOperation;
+		
+		explicit FDelayedDragData(
+			const FVector2D& InInitialPosition, const FKey& InEffectiveKey, const FLatticeBounds& InBounds, const FLatticeHoverState& InHoverState
+			)
+			: DelayedDrag(InInitialPosition, InEffectiveKey), Bounds(InBounds), HoverState(InHoverState)
+		{}
+	};
+	TOptional<FDelayedDragData> DelayedDragData;
+
+	/**
+	 * We store the shape of the lattice grid before starting the drag operation.
+	 * This way, we can restore the shape of the grid when the user undo-es.
+	 */
+	TObjectPtr<UCurveEditorTools_LatticeUndoData> UndoObject;
+	
+	/** Creates or destroys LatticeDeformer based on the selection. */
+	void UpdateDeformerBounds(const TSharedRef<FCurveEditor>& InCurveEditor);
+
+	/** Updates the bounds after the user has externally moved the keys, e.g. MMB move keys, etc. */
+	void OnCurvesModified();
+	/** Call this if the keys have externally moved. Tries to fit the lattice grid over the keys or reconstructs a new rectangular one if not possible. */
+	void AdjustLatticeGridOrRecompute(const TSharedRef<FCurveEditor>& InCurveEditor);
+	
+	/** Stops the active drag operation. */
+	void CancelDrag();
+
+	/** Recomputes all draw data. */
+	void RecomputeDrawData(const FGeometry& InViewContainerGeometry, const TSharedPtr<FCurveEditor>& InCurveEditor);
+	/** Populates the control points and edges. */
+	TArray<FVector2D> ComputeDrawControlPoints(const FCurveEditor& InCurveEditor) const;
+	TArray<FLatticeControlEdge> ComputeDrawControlEdges(TConstArrayView<FVector2D> InControlPoints) const;
+
+	TUniquePtr<FLatticeDragOp> CreateDragDrop(const FLatticeBounds& InBounds, const FLatticeHoverState& InHoverState);
+	/** Generic function for updating control points. */
+	void OnSetControlPoints(TConstArrayView<int32> InControlPointsToUpdate, TConstArrayView<FVector2D> InNewControlPoints);
+	
+	const FLatticeDeformer2D& GetDeformer() const { check(LatticeDeformerState); return LatticeDeformerState->Deformer; }
+	FLatticeDeformer2D& GetDeformer() { check(LatticeDeformerState); return LatticeDeformerState->Deformer; }
+};
+}
diff --git a/Engine/Plugins/Editor/CurveEditorTools/Source/CurveEditorTools/Private/Tools/Lattice/Drag/LatticeDragOp.cpp b/Engine/Plugins/Editor/CurveEditorTools/Source/CurveEditorTools/Private/Tools/Lattice/Drag/LatticeDragOp.cpp
new file mode 100644
--- /dev/null
+++ b/Engine/Plugins/Editor/CurveEditorTools/Source/CurveEditorTools/Private/Tools/Lattice/Drag/LatticeDragOp.cpp
@@ -0,0 +1,69 @@
+﻿// Copyright Epic Games, Inc. All Rights Reserved.
+
+#include "LatticeDragOp.h"
+
+#include "CurveEditor.h"
+
+namespace UE::CurveEditorTools
+{
+FLatticeDragOp::FLatticeDragOp(TWeakPtr<FCurveEditor> InCurveEditor)
+	: CurveEditor(MoveTemp(InCurveEditor))
+	, InitialMousePosition()
+{}
+
+void FLatticeDragOp::BeginDrag(const FGeometry& InGeometry, const FPointerEvent& InMouseEvent, const FVector2D& InInitialMousePosition)
+{
+	InitialMousePosition = LastMousePosition = InMouseEvent.GetScreenSpacePosition();
+	AccumulateMouseMovement(InGeometry, InMouseEvent);
+	OnBeginDrag(InGeometry, InitialMousePosition);
+}
+
+void FLatticeDragOp::MoveMouse(const FGeometry& InGeometry, const FPointerEvent& InMouseEvent)
+{
+	AccumulateMouseMovement(InGeometry, InMouseEvent);
+}
+
+void FLatticeDragOp::FinishedPointerInput()
+{
+	if (AccumulatedMouseMovement)
+	{
+		OnMoveMouse(AccumulatedMouseMovement->CachedGeometry, AccumulatedMouseMovement->AccumulatedPosition);
+		AccumulatedMouseMovement.Reset();
+	}
+}
+
+void FLatticeDragOp::EndDrag(const FGeometry& InGeometry, const FPointerEvent& InMouseEvent)
+{
+	OnEndDrag(InGeometry, InMouseEvent.GetScreenSpacePosition());
+	AccumulatedMouseMovement.Reset();
+}
+
+void FLatticeDragOp::CancelDrag()
+{
+	OnCancelDrag();
+}
+
+void FLatticeDragOp::AccumulateMouseMovement(const FGeometry& InGeometry, const FPointerEvent& InMouseEvent)
+{
+	const TSharedPtr<FCurveEditor> CurveEditorPin = CurveEditor.Pin();
+	if (!CurveEditorPin)
+	{
+		return;
+	}
+	
+	const FVector2D ScreenSpacePosition = InMouseEvent.GetScreenSpacePosition();
+	const FVector2D SnappedPosition = CurveEditorPin->GetAxisSnap().GetSnappedPosition(
+		InitialMousePosition, ScreenSpacePosition, LastMousePosition, InMouseEvent, SnapState
+		);
+	LastMousePosition = ScreenSpacePosition;
+	
+	if (!AccumulatedMouseMovement)
+	{
+		AccumulatedMouseMovement.Emplace(InGeometry, SnappedPosition);
+	}
+	else
+	{
+		AccumulatedMouseMovement->AccumulatedPosition = SnappedPosition;
+	}
+}
+}
diff --git a/Engine/Plugins/Editor/CurveEditorTools/Source/CurveEditorTools/Private/Tools/Lattice/Drag/LatticeDragOp.h b/Engine/Plugins/Editor/CurveEditorTools/Source/CurveEditorTools/Private/Tools/Lattice/Drag/LatticeDragOp.h
new file mode 100644
--- /dev/null
+++ b/Engine/Plugins/Editor/CurveEditorTools/Source/CurveEditorTools/Private/Tools/Lattice/Drag/LatticeDragOp.h
@@ -0,0 +1,69 @@
+﻿// Copyright Epic Games, Inc. All Rights Reserved.
+
+#pragma once
+
+#include "CurveEditorSnapMetrics.h"
+#include "Layout/Geometry.h"
+#include "Math/MathFwd.h"
+#include "Math/Vector2D.h"
+#include "Misc/Optional.h"
+#include "Templates/SharedPointer.h"
+
+class FCurveEditor;
+
+namespace UE::CurveEditorTools
+{
+/** Handles updating the lattice deformer's control point(s) in response to a drag. */
+class FLatticeDragOp
+{
+public:
+
+	explicit FLatticeDragOp(TWeakPtr<FCurveEditor> InCurveEditor);
+
+	void BeginDrag(const FGeometry& InGeometry, const FPointerEvent& InMouseEvent, const FVector2D& InInitialMousePosition);
+	void MoveMouse(const FGeometry& InGeometry, const FPointerEvent& InMouseEvent);
+	void FinishedPointerInput();
+	void EndDrag(const FGeometry& InGeometry, const FPointerEvent& InMouseEvent);
+	void CancelDrag();
+
+	virtual ~FLatticeDragOp() = default;
+
+protected:
+
+	virtual void OnBeginDrag(const FGeometry& InGeometry, const FVector2D& InInitialMousePosition) {}
+	virtual void OnMoveMouse(const FGeometry& InGeometry, const FVector2D& InScreenPosition) {}
+	virtual void OnEndDrag(const FGeometry& InGeometry, const FVector2D& InMousePosition) {}
+	virtual void OnCancelDrag() {}
+
+	TWeakPtr<FCurveEditor> GetCurveEditor() const { return CurveEditor; }
+	FVector2D GetInitialMousePosition() const { return InitialMousePosition; }
+	
+private:
+
+	/** Used to get snapping settings. */
+	const TWeakPtr<FCurveEditor> CurveEditor;
+	
+	/** Set by OnBeginDrag. */
+	FVector2D InitialMousePosition;
+	/** Used by the snap logic. Set every time mouse moved. */
+	FVector2D LastMousePosition;
+
+	/** Passed to the snap logic */
+	FCurveEditorAxisSnap::FSnapState SnapState;
+	
+	struct FAccumulatedMouseMovement
+	{
+		const FGeometry CachedGeometry;
+		FVector2D AccumulatedPosition;
+
+		explicit FAccumulatedMouseMovement(const FGeometry& InCachedGeometry, const FVector2D& InInitialPosition)
+			: CachedGeometry(InCachedGeometry)
+			, AccumulatedPosition(InInitialPosition)
+		{}
+	};
+	TOptional<FAccumulatedMouseMovement> AccumulatedMouseMovement;
+
+	void AccumulateMouseMovement(const FGeometry& InGeometry, const FPointerEvent& InMouseEvent);
+};
+}
+
diff --git a/Engine/Plugins/Editor/CurveEditorTools/Source/CurveEditorTools/Private/Tools/Lattice/Drag/LatticeDragOp_MoveControlPoints.cpp b/Engine/Plugins/Editor/CurveEditorTools/Source/CurveEditorTools/Private/Tools/Lattice/Drag/LatticeDragOp_MoveControlPoints.cpp
new file mode 100644
--- /dev/null
+++ b/Engine/Plugins/Editor/CurveEditorTools/Source/CurveEditorTools/Private/Tools/Lattice/Drag/LatticeDragOp_MoveControlPoints.cpp
@@ -0,0 +1,40 @@
+﻿// Copyright Epic Games, Inc. All Rights Reserved.
+
+#include "LatticeDragOp_MoveControlPoints.h"
+
+#include "Templates/SharedPointer.h"
+#include "Tools/Lattice/Misc/LatticeDrawUtils.h"
+
+namespace UE::CurveEditorTools
+{
+FLatticeDragOp_MoveControlPoints::FLatticeDragOp_MoveControlPoints(
+	TWeakPtr<FCurveEditor> InCurveEditor,
+	FControlPointArray InInitialControlPoints,
+	const FLatticeBounds& InBounds,
+	FOnControlPointsMoved SetControlPointDelegate
+	)
+	: FLatticeDragOp(MoveTemp(InCurveEditor))
+	, InitialControlPoints(MoveTemp(InInitialControlPoints))
+	, Bounds(InBounds)
+	, SetControlPointDelegate(MoveTemp(SetControlPointDelegate))
+	, UpdatedControlPointPositions(InitialControlPoints)
+{}
+	
+void FLatticeDragOp_MoveControlPoints::OnMoveMouse(const FGeometry& InGeometry, const FVector2D& InScreenPosition)
+{
+	const FSlateLayoutTransform ScreenToWidget = InGeometry.GetAccumulatedLayoutTransform().Inverse();
+	const FVector2D Initial_WidgetSpace = ScreenToWidget.TransformPoint(GetInitialMousePosition());
+	const FVector2D New_WidgetSpace = ScreenToWidget.TransformPoint(InScreenPosition);
+		
+	const FVector2D DeltaMove_WidgetSpace = New_WidgetSpace - Initial_WidgetSpace;
+	FVector2D ValuePerSlate = (Bounds.MaxValues - Bounds.MinValues) / (Bounds.MaxSlatePosition - Bounds.MinSlatePosition);
+	ValuePerSlate.Y *= -1.f; // Mouse.Y and Screen.Y move in opposite directions.
+	const FVector2D DeltaMove_ControlPointSpace = DeltaMove_WidgetSpace * ValuePerSlate;
+
+	for (int32 Index = 0; Index < InitialControlPoints.Num(); ++Index)
+	{
+		UpdatedControlPointPositions[Index] = InitialControlPoints[Index] + DeltaMove_ControlPointSpace;
+	}
+	SetControlPointDelegate.Execute(UpdatedControlPointPositions);
+}
+}
diff --git a/Engine/Plugins/Editor/CurveEditorTools/Source/CurveEditorTools/Private/Tools/Lattice/Drag/LatticeDragOp_MoveControlPoints.h b/Engine/Plugins/Editor/CurveEditorTools/Source/CurveEditorTools/Private/Tools/Lattice/Drag/LatticeDragOp_MoveControlPoints.h
new file mode 100644
--- /dev/null
+++ b/Engine/Plugins/Editor/CurveEditorTools/Source/CurveEditorTools/Private/Tools/Lattice/Drag/LatticeDragOp_MoveControlPoints.h
@@ -0,0 +1,52 @@
+﻿// Copyright Epic Games, Inc. All Rights Reserved.
+
+#pragma once
+
+#include "Delegates/Delegate.h"
+#include "HAL/Platform.h"
+#include "LatticeDragOp.h"
+#include "Math/Vector2D.h"
+#include "Templates/SharedPointerFwd.h"
+#include "Tools/Lattice/Misc/LatticeDrawUtils.h"
+
+class FCurveEditor;
+
+namespace UE::CurveEditorTools
+{
+struct FLatticeBounds;
+using FControlPointArray = TArray<FVector2D, TInlineAllocator<4>>;
+	
+/** Given an array of control points, invokes a delegate with the update control point locations based on the mouse movement. */
+class FLatticeDragOp_MoveControlPoints : public FLatticeDragOp
+{
+public:
+
+	DECLARE_DELEGATE_OneParam(FOnControlPointsMoved, const TConstArrayView<FVector2D>& InNewControlPointPositions);
+
+	explicit FLatticeDragOp_MoveControlPoints(
+		TWeakPtr<FCurveEditor> InCurveEditor,
+		FControlPointArray InInitialControlPoints,
+		const FLatticeBounds& InBounds,
+		FOnControlPointsMoved SetControlPointDelegate
+		);
+
+protected:
+	
+	//~ Begin FLatticeDragOp Interface
+	virtual void OnMoveMouse(const FGeometry& InGeometry, const FVector2D& InScreenPosition) override;
+	//~ End FLatticeDragOp Interface
+
+private:
+
+	/** Initial positions of the control points. */
+	const FControlPointArray InitialControlPoints;
+	/** Used to transform points */
+	const FLatticeBounds Bounds;
+	/** Invoked to change the control points by the. */
+	const FOnControlPointsMoved SetControlPointDelegate;
+
+	/** Updated every time the mouse moves: InitialControlPoints + DeltaMouseMove. */
+	FControlPointArray UpdatedControlPointPositions;
+};
+}
+
diff --git a/Engine/Plugins/Editor/CurveEditorTools/Source/CurveEditorTools/Private/Tools/Lattice/LatticeControlsDrawData.h b/Engine/Plugins/Editor/CurveEditorTools/Source/CurveEditorTools/Private/Tools/Lattice/LatticeControlsDrawData.h
new file mode 100644
--- /dev/null
+++ b/Engine/Plugins/Editor/CurveEditorTools/Source/CurveEditorTools/Private/Tools/Lattice/LatticeControlsDrawData.h
@@ -0,0 +1,56 @@
+﻿// Copyright Epic Games, Inc. All Rights Reserved.
+
+#pragma once
+
+#include "Containers/Array.h"
+#include "Math/MathFwd.h"
+#include "Math/IntPoint.h"
+#include "Math/Vector2D.h"
+#include "Misc/Optional.h"
+#include "HAL/Platform.h"
+
+namespace UE::CurveEditorTools
+{
+struct FLatticeControlEdge;
+
+/** Describes what is hovered in the control widget. */
+struct FLatticeHoverState
+{
+	/** Index to ControlPointMatrix. Indicates hovered control point. */
+	TOptional<int32> HoveredControlPoint;
+	/**
+	 * Index to edge.
+	 * @see Indexing of edges in FLatticeDeformer2D.
+	 */
+	TOptional<int32> HoveredEdge;
+
+	/**
+	 * Index to the cell.
+	 * @see Indexing of cells in FLatticeDeformer2D. 
+	 */
+	TOptional<int32> HoveredCell;
+
+	bool IsHovered() const { return HoveredControlPoint || HoveredEdge || HoveredCell; }
+};
+	
+/**
+ * Holds the data required to draw a lattice control widget.
+ *
+ * This struct represents the exact information the drawing algorithm requires. The intention is that the drawing thus becomes as straighforward
+ * as possible. This way, if we wanted to, we could unit test the produced FLatticeControlsDrawData.
+ */
+struct FLatticeControlsDrawData
+{
+	/** The control points in the space of the SCurveEditorViewContainer (i.e. what's passed into ICurveEditorToolExtension::OnPaint). */
+	TArray<FVector2D> ControlPoints;
+	/** The edges to draw */
+	TArray<FLatticeControlEdge> ControlEdges;
+
+	/** Number of points in x-direction. Correlates with ControlPoints.Num(). */
+	int32 MatrixWidth;
+
+	/** Determines which index in ControlPoints or ControlEdges should be drawn as hovered. */
+	FLatticeHoverState HoverState;
+};
+}
+
diff --git a/Engine/Plugins/Editor/CurveEditorTools/Source/CurveEditorTools/Private/Tools/Lattice/LatticeDeformer2D.cpp b/Engine/Plugins/Editor/CurveEditorTools/Source/CurveEditorTools/Private/Tools/Lattice/LatticeDeformer2D.cpp
new file mode 100644
--- /dev/null
+++ b/Engine/Plugins/Editor/CurveEditorTools/Source/CurveEditorTools/Private/Tools/Lattice/LatticeDeformer2D.cpp
@@ -0,0 +1,285 @@
+﻿// Copyright Epic Games, Inc. All Rights Reserved.
+
+#include "LatticeDeformer2D.h"
+
+#include "Algo/AllOf.h"
+#include "Async/ParallelFor.h"
+#include "Misc/LatticeUtils.h"
+#include "HAL/Platform.h"
+
+namespace UE::CurveEditorTools::Math
+{
+/** @return The UV coordinates InPoint has in the square formed by InBottomLeft and InTopRight. */
+static FVector2D ComputeCoordinatesInCell(const FVector2D& InPoint, const FVector2D& InBottomLeft, const FVector2D& InTopRight)
+{
+	const double MinX = InBottomLeft.X;
+	const double MinY = InBottomLeft.Y;
+	const double Width = InTopRight.X - MinX;
+	const double Height = InTopRight.Y - MinY;
+
+	// Division by 0 should not occur because we check for 0 width and height in constructor.
+	// However, maybe a future change breaks this assumption (fix it!), or bad API usage calls UpdateControlPoints before AddPoints_BeforeLatticeMove.
+	const bool bNonZeroWidth = !FMath::IsNearlyZero(Width);
+	const bool bNonZeroHeight =!FMath::IsNearlyZero(Height);
+	checkSlow(bNonZeroWidth && bNonZeroHeight);
+	
+	const double U = LIKELY(bNonZeroWidth) ? (InPoint.X - MinX) / Width : 0.0;
+	const double V = LIKELY(bNonZeroHeight) ? (InPoint.Y - MinY) / Height : 0.0;
+	
+	return { U, V };
+}
+}
+
+namespace UE::CurveEditorTools
+{
+FLatticeDeformer2D::FLatticeDeformer2D(int32 InNumCellsInX, int32 InNumCellsInY, const FVector2D& InBottomLeft, const FVector2D& InTopRight)
+	: CellDimensions(InNumCellsInX, InNumCellsInY)
+{
+	// Lattice with 0 width or height does not make sense. We'd divide by 0 in ComputeCoordinatesInCell.
+	if (!ensure(InBottomLeft.X < InTopRight.X) || !ensure(InBottomLeft.Y < InTopRight.Y))
+	{
+		// User specified invalid bounds. The lattice will be empty.
+		return;
+	}
+	ControlPoints.SetNumUninitialized(NumPointsInWidth() * NumPointsInHeight());
+
+	const double DeltaX = (InTopRight.X - InBottomLeft.X) / static_cast<double>(InNumCellsInX);
+	const double DeltaY = (InTopRight.Y - InBottomLeft.Y) / static_cast<double>(InNumCellsInY);
+	for (int32 X = 0; X < NumPointsInWidth(); ++X)
+	{
+		for (int32 Y = 0; Y < NumPointsInWidth(); ++Y)
+		{
+			const int32 Index = MatrixIndicesToFlatIndex(X, Y, NumPointsInWidth());
+			const FVector2D ControlPosition(InBottomLeft.X + DeltaX * X, InBottomLeft.Y + DeltaY * Y);
+			ControlPoints[Index] = ControlPosition;
+		}
+	}
+
+	Cells.SetNum(NumCells());
+}
+
+void FLatticeDeformer2D::AddPoints_BeforeLatticeMoved(TConstArrayView<FKeyInCurveId> InKeys, TConstArrayView<FVector2D> InPoints)
+{
+	if (!ensure(InKeys.Num() == InPoints.Num()))
+	{
+		return;
+	}
+
+	// No need to make this ParallelFor: 4500 keys take about 0.6ms.
+	for (int32 KeyIndex = 0; KeyIndex < InKeys.Num(); ++KeyIndex)
+	{
+		const FVector2D& KeyPoint = InPoints[KeyIndex];
+		// Check whether in cell is relatively cheap (addition & multiplication) ...
+		const int32 CellIndex = FindCellContainingPoint(KeyPoint);
+		if (!Cells.IsValidIndex(CellIndex))
+		{
+			continue;
+		}
+		
+		TArray<FLatticeKeyCoords>& Keys = Cells[CellIndex].KeyCoords;
+		TArray<FKeyInCurveId>& KeyIds = Cells[CellIndex].KeyIds;
+		KeyIds.Emplace(InKeys[KeyIndex]);
+		Keys.Emplace(ComputeSingleCellWeights(GetCellIndices(CellIndex), InPoints[KeyIndex]));
+	}
+}
+
+void FLatticeDeformer2D::UpdateControlPoints(TConstArrayView<int32> InControlPointsToUpdate, TConstArrayView<FVector2D> InNewControlPoints,
+	const TFunctionRef<void(const FKeyInCurveId&, const FVector2D&)> HandleKeyChange)
+{
+	const bool bValidIndices = Algo::AllOf(InControlPointsToUpdate, [this](int32 InIndex){ return ControlPoints.IsValidIndex(InIndex); });
+	if (!ensure(bValidIndices && InControlPointsToUpdate.Num() == InNewControlPoints.Num())
+		|| InControlPointsToUpdate.IsEmpty())
+	{
+		return;
+	}
+	
+	if (NumCells() == 0)
+	{
+		SetControlsPointNoRecompute(InControlPointsToUpdate, InNewControlPoints);
+		
+		// Fast path for common 1x1 lattice case
+		RecomputeCell(0, HandleKeyChange);
+	}
+	else
+	{
+		TArray<int32> AffectedCells;
+		for (int32 Index = 0; Index < InControlPointsToUpdate.Num(); ++Index)
+		{
+			const int32 ControlPointIndex = InControlPointsToUpdate[Index];
+			ControlPoints[ControlPointIndex] = InNewControlPoints[Index];
+
+			// Get the cell indices of the cell to the top-left, etc. of the corner.
+			const auto[TopLeft, TopRight, BottomRight, BottomLeft] = GetMatrixCellSlotsWithCorner(
+				ControlPointIndex, NumPointsInWidth(), NumPointsInHeight()
+				);
+			// Need to recompute every cell that shares corner with the control point.
+			AffectedCells.AddUnique(TopLeft);
+			AffectedCells.AddUnique(TopRight);
+			AffectedCells.AddUnique(BottomRight);
+			AffectedCells.AddUnique(BottomLeft);
+		}
+
+		for (const int32& CellIndex : AffectedCells)
+		{
+			RecomputeCell(CellIndex, HandleKeyChange);
+		}
+	}
+}
+
+void FLatticeDeformer2D::SetControlsPointNoRecompute(TConstArrayView<int32> InControlPointsToUpdate, TConstArrayView<FVector2D> InNewControlPoints)
+{
+	const bool bValidIndices = Algo::AllOf(InControlPointsToUpdate, [this](int32 InIndex){ return ControlPoints.IsValidIndex(InIndex); });
+	if (!ensure(bValidIndices && InControlPointsToUpdate.Num() == InNewControlPoints.Num())
+		|| InControlPointsToUpdate.IsEmpty())
+	{
+		return;
+	}
+	
+	for (int32 Index = 0; Index < InControlPointsToUpdate.Num(); ++Index)
+	{
+		ControlPoints[InControlPointsToUpdate[Index]] = InNewControlPoints[Index];
+	}
+}
+
+FLatticeControlEdge FLatticeDeformer2D::GetControlEdge(int32 InEdgeIndex) const
+{
+	const auto[Start, End] = GetEdgeIndices(InEdgeIndex);
+	if (ensure(ControlPoints.IsValidIndex(Start) && ControlPoints.IsValidIndex(End)))
+	{
+		return FLatticeControlEdge(ControlPoints[Start], ControlPoints[End]);
+	}
+	return FLatticeControlEdge();
+}
+
+FEdgeVertexIndices FLatticeDeformer2D::GetEdgeIndices(int32 InEdgeIndex) const
+{
+	const int32 Width = NumPointsInWidth();
+	const int32 Height = NumPointsInHeight();
+	const bool bIsInRange = Width > 0 && Height > 0 && InEdgeIndex >= 0 && InEdgeIndex < NumEdges();
+	return ensure(bIsInRange) ? CurveEditorTools::GetEdgeIndices(InEdgeIndex, Width, Height) : FEdgeVertexIndices();
+}
+
+FCellVertexIndices FLatticeDeformer2D::GetCellIndices(int32 InCellIndex) const
+{
+	const bool bIsInRange = InCellIndex >= 0 && InCellIndex < NumCells();
+	return ensure(bIsInRange) ? CurveEditorTools::GetCellIndices(InCellIndex, NumPointsInWidth()) : FCellVertexIndices();
+}
+
+void FLatticeDeformer2D::SetKeyIdsInCell(int32 InCellIndex, TArray<FKeyInCurveId> InNewKeyCoords)
+{
+	if (ensure(Cells.IsValidIndex(InCellIndex)))
+	{
+		Cells[InCellIndex].KeyIds = MoveTemp(InNewKeyCoords);
+	}
+}
+
+TConstArrayView<FKeyInCurveId> FLatticeDeformer2D::GetKeyIdsInCell(int32 InCellIndex) const
+{
+	return ensure(Cells.IsValidIndex(InCellIndex)) ? Cells[InCellIndex].KeyIds : TConstArrayView<FKeyInCurveId>();
+} 
+
+void FLatticeDeformer2D::SetKeyCoordsInCell(int32 InCellIndex, TArray<FLatticeKeyCoords> InNewKeyCoords)
+{
+	if (ensure(Cells.IsValidIndex(InCellIndex)))
+	{
+		Cells[InCellIndex].KeyCoords = MoveTemp(InNewKeyCoords);
+	}
+}
+
+TConstArrayView<FLatticeKeyCoords> FLatticeDeformer2D::GetKeyCoordsInCell(int32 InCellIndex) const
+{
+	return ensure(Cells.IsValidIndex(InCellIndex)) ? Cells[InCellIndex].KeyCoords : TConstArrayView<FLatticeKeyCoords>();
+}
+
+int32 FLatticeDeformer2D::FindCellContainingPoint(const FVector2D& InPoint) const
+{
+	if (NumCells() <= 1)
+	{
+		// Fast-path
+		return 0; // If there are 0 cells, Cells.IsValidIndex(0) will still return false
+	}
+	else
+	{
+		const FVector2D& BottomLeft = ControlPoints[MatrixIndicesToFlatIndex(0, NumPointsInHeight() - 1, NumPointsInWidth())];
+		const FVector2D& TopRight = ControlPoints[NumPointsInWidth() - 1];
+		const FVector2D& UV = Math::ComputeCoordinatesInCell(InPoint, BottomLeft, TopRight);
+		const bool bIsInGrid = 0.0 <= UV.GetMin() && UV.GetMax() <= 1.0;
+		if (!bIsInGrid)
+		{
+			return INDEX_NONE;
+		}
+		
+		const int32 Column = static_cast<int32>(UV.X / NumPointsInWidth());
+		const int32 Row = static_cast<int32>(UV.Y / NumPointsInHeight());
+		return MatrixIndicesToFlatIndex(Column, Row, NumCellsInWidth());
+	}
+}
+
+int32 FLatticeDeformer2D::NumEdges() const
+{
+	if (NumPointsInWidth() == 0 || NumPointsInHeight() == 0)
+	{
+		return 0; 
+	}
+
+	const int32 HorizontalEdges = (NumPointsInWidth() - 1) * NumPointsInHeight();
+	const int32 VerticalEdges = (NumPointsInHeight() - 1) * NumPointsInWidth();
+
+	// Total number of edges
+	return HorizontalEdges + VerticalEdges;
+}
+
+int32 FLatticeDeformer2D::NumCells() const
+{
+	return (NumPointsInWidth() - 1) * (NumPointsInHeight() - 1);
+}
+
+void FLatticeDeformer2D::RecomputeCell(int32 InCellIndex, const TFunctionRef<void(const FKeyInCurveId&, const FVector2D&)> HandleKeyChange)
+{
+	if (!Cells.IsValidIndex(InCellIndex))
+	{
+		return;
+	}
+		
+	const FCellVertexIndices CellVertices = GetCellIndices(InCellIndex);
+
+	// This takes about 0.083 ms for 9000 keys.
+	TArray<FLatticeKeyCoords>& Keys = Cells[InCellIndex].KeyCoords;
+	TArray<FVector2D> Positions;
+	Positions.SetNumUninitialized(Keys.Num());
+	ParallelFor(Keys.Num(), [this, &CellVertices, &Keys, &Positions](int32 KeyIndex)
+	{
+		FLatticeKeyCoords& KeyData = Keys[KeyIndex];
+		Positions[KeyIndex] = RecomputeKeyValue(KeyData, CellVertices);
+	});
+
+	TArray<FKeyInCurveId>& KeyIds = Cells[InCellIndex].KeyIds;
+	for (int32 Index = 0; Index < Keys.Num(); ++Index)
+	{
+		HandleKeyChange(KeyIds[Index], Positions[Index]);
+	}
+}
+
+FVector2D FLatticeDeformer2D::RecomputeKeyValue(const FLatticeKeyCoords& PointCoords, const FCellVertexIndices& InIndices) const
+{
+	const FVector2D Inverse(1.0 - PointCoords.X, 1.0 - PointCoords.Y);
+
+	// Compute the new point based on the coordinates it had in the original lattice...
+	// FinalPoint = (1-U)*(1-V)*BottomLeft + (1-U)*V*TopLeft + U*(1-V)*BottomRight + U*V*TopRight
+	const FVector2D BottomLeft = Inverse.X * Inverse.Y * ControlPoints[InIndices.BottomLeft];
+	const FVector2D TopLeft = Inverse.X * PointCoords.Y * ControlPoints[InIndices.TopLeft];
+	const FVector2D BottomRight = PointCoords.X * Inverse.Y * ControlPoints[InIndices.BottomRight];
+	const FVector2D TopRight = PointCoords.X * PointCoords.Y * ControlPoints[InIndices.TopRight];
+
+	const FVector2D FinalPoint = BottomLeft + TopLeft + BottomRight + TopRight;
+	return { FinalPoint.X, FinalPoint.Y };
+}
+
+FLatticeKeyCoords FLatticeDeformer2D::ComputeSingleCellWeights(const FCellVertexIndices& InIndices, const FVector2D& InPoint) const
+{
+	const auto[TopLeft, TopRight, BottomRight, BottomLeft] = InIndices;
+	const FVector2D Vertex_TopRight = ControlPoints[TopRight];
+	const FVector2D Vertex_BottomLeft = ControlPoints[BottomLeft];
+	return Math::ComputeCoordinatesInCell(InPoint, Vertex_BottomLeft, Vertex_TopRight);
+}
+}
\ No newline at end of file
diff --git a/Engine/Plugins/Editor/CurveEditorTools/Source/CurveEditorTools/Private/Tools/Lattice/LatticeDeformer2D.h b/Engine/Plugins/Editor/CurveEditorTools/Source/CurveEditorTools/Private/Tools/Lattice/LatticeDeformer2D.h
new file mode 100644
--- /dev/null
+++ b/Engine/Plugins/Editor/CurveEditorTools/Source/CurveEditorTools/Private/Tools/Lattice/LatticeDeformer2D.h
@@ -0,0 +1,175 @@
+﻿// Copyright Epic Games, Inc. All Rights Reserved.
+
+#pragma once
+
+#include "Containers/Array.h"
+#include "Containers/ContainersFwd.h"
+#include "Curves/KeyHandle.h"
+#include "CurveDataAbstraction.h"
+#include "HAL/Platform.h"
+#include "Math/MathFwd.h"
+#include "Math/Vector2D.h"
+#include "Templates/FunctionFwd.h"
+
+namespace UE::CurveEditorTools
+{
+struct FEdgeVertexIndices;
+struct FCellVertexIndices;
+	
+/**  */
+struct FLatticeControlEdge
+{
+	/** This allows us to convert TArray<FLatticeControlEdge> to TArray<FVector2D> by just taking the pointer to Points. */
+	FVector2D Points[2];
+
+	const FVector2D& Start() const { return Points[0]; }
+	const FVector2D& End() const { return Points[1]; }
+	
+	FVector2D& Start() { return Points[0]; }
+	FVector2D& End() { return Points[1]; }
+
+	explicit FLatticeControlEdge() : FLatticeControlEdge(FVector2D::ZeroVector, FVector2D::ZeroVector) {}
+	explicit FLatticeControlEdge(const FVector2D& Start, const FVector2D& End)
+		: Points{ Start, End }
+	{}
+};
+
+struct FKeyInCurveId
+{
+	FKeyHandle KeyId;
+	FCurveModelID CurveId;
+	
+	friend FArchive& operator<<(FArchive& InArchive, FKeyInCurveId& InCurveData)
+	{
+		InArchive << InCurveData.KeyId;
+		InArchive << InCurveData.CurveId;
+		return InArchive;
+	}
+};
+
+/** Identifies the normalized coordinates a key has in a cell. */
+using FLatticeKeyCoords = FVector2D;
+	
+/**
+ * Implements a lattice deformer for 2D operations.
+ *
+ * ===== Indexing ====
+ * Control points:
+ *  0---1---2---3
+ *	|   |   |   |
+ *	4---5---6---7
+ *	|   |   |   |
+ *	8---9--10--11
+ * 
+ * Below, x are control points and - / I are edges.
+ * Edge indices are numbered by rows then by columns, i.e.
+ *  x--0--x--1--x
+ *  6     7     8
+ *  x--2--x--3--x  
+ *  9     10    11
+ *  x--4--x--5--x 
+ *
+ * Cells indices are numbered from top-left to bottom-right, i.e.
+ * x---x---x
+ * I 0 I 1 I
+ * x---x---x
+ *
+ * In the future, if FLatticeDeformer2D needs to be generalized / extracted from this module, FKeyInCurveId can become a templated parameter since
+ * FLatticeDeformer2D does not directly depend on it.
+ */
+class FLatticeDeformer2D
+{
+public:
+
+	/** Completely resets this deformer and sets a new grid.  */
+	explicit FLatticeDeformer2D(int32 InNumCellsInX, int32 InNumCellsInY, const FVector2D& InBottomLeft, const FVector2D& InTopRight);
+
+	/**
+	 * Computes the weights for each of the points.
+	 * 
+	 * This function assumes that none of the control points have yet been moved: if so, the computed weights are incorrect.
+	 * The computation of the weights requires that all cells are rectangular.
+	 */
+	void AddPoints_BeforeLatticeMoved(TConstArrayView<FKeyInCurveId> InKeys, TConstArrayView<FVector2D> InPoints);
+	
+	/** Moves the specified control points and recomputes the affected points. */
+	void UpdateControlPoints(TConstArrayView<int32> InControlPointsToUpdate, TConstArrayView<FVector2D> InNewControlPoints,
+		const TFunctionRef<void(const FKeyInCurveId&, const FVector2D&)> HandleKeyChange
+		);
+
+	/** Sets the value of the control points without recomputing the values of the keys. Useful if you are resetting to a previous state. */
+	void SetControlsPointNoRecompute(TConstArrayView<int32> InControlPointsToUpdate, TConstArrayView<FVector2D> InNewControlPoints);
+
+	/** @return The control points, which is a flattened matrix. */
+	TConstArrayView<FVector2D> GetControlPoints() const { return ControlPoints; }
+	/** @param InIndex 0 <= InIndex < GetNumEdges() */
+	FLatticeControlEdge GetControlEdge(int32 InEdgeIndex) const;
+	
+	/**
+	 * @param InEdgeIndex 0 <= InEdgeIndex < NumEdges().
+	 * @return Indices to GetControlPoints or all INDEX_NONE if invalid.
+	 */
+	FEdgeVertexIndices GetEdgeIndices(int32 InEdgeIndex) const;
+	/**
+	 * @param InCellIndex 0 <= InCellIndex < NumCells()
+	 * @return Indices to GetControlPoints or all INDEX_NONE if invalid.
+	 */
+	FCellVertexIndices GetCellIndices(int32 InCellIndex) const;
+
+	/** Sets the key IDs of the keys in the cell. */
+	void SetKeyIdsInCell(int32 InCellIndex, TArray<FKeyInCurveId> InNewKeyCoords);
+	/** @return The key IDs of the keys in the cell. */
+	TConstArrayView<FKeyInCurveId> GetKeyIdsInCell(int32 InCellIndex) const;
+
+	/** Sets the key coords of keys in the cell. This is useful for restoring previous data. */
+	void SetKeyCoordsInCell(int32 InCellIndex, TArray<FLatticeKeyCoords> InNewKeyCoords);
+	/** @return The coordinates of the keys in the cell. */
+	TConstArrayView<FLatticeKeyCoords> GetKeyCoordsInCell(int32 InCellIndex) const;
+	
+	/** @return The vertices of the cell containing InPoint, or all INDEX_NONE if not contained.  */
+	int32 FindCellContainingPoint(const FVector2D& InPoint) const;
+	
+	int32 NumControlPoints() const { return ControlPoints.Num(); }
+	int32 NumPointsInWidth() const { return CellDimensions.X + 1; }
+	int32 NumPointsInHeight() const { return CellDimensions.Y + 1; }
+	int32 NumEdges() const;
+	int32 NumCellsInWidth() const { return CellDimensions.X; }
+	int32 NumCellsInHeight() const { return CellDimensions.Y; }
+	int32 NumCells() const;
+	
+private:
+	
+	/** The number of cells in X and Y direction. */
+	const FIntPoint CellDimensions;
+	
+	/**
+	 * The control points of the deformer. Moving a control point causes the keys it affects to change transform as well.
+	 * 
+	 * This is a flat matrix of dimension CellDimensions.X + 1 * CellDimensions.Y + 1,
+	 * i.e. the 1st point in the 2nd row would be ControlPointMatrix[CellDimensions.X + 2]
+	 */
+	TArray<FVector2D> ControlPoints;
+
+	struct FCellData
+	{
+		TArray<FKeyInCurveId> KeyIds;
+		TArray<FLatticeKeyCoords> KeyCoords;
+	};
+	/**
+	 * Caches which points belong to which cells. Moving a control point causes the keys it affects to change transform as well.
+	 * 
+	 * This is a flat matrix of dimension CellDimensions.X * CellDimensions.Y ,
+	 * i.e. the 1st point in the 2nd row would be ControlPointMatrix[CellDimensions.X + 1]
+	 */
+	TArray<FCellData> Cells;
+
+	/** Recomputes all key positions in the cell. Call this when cell has changed its control points. */
+	void RecomputeCell(int32 InCellIndex, const TFunctionRef<void(const FKeyInCurveId&, const FVector2D&)> HandleKeyChange);
+
+	/** Computes the value the key should by computing the weighted linear product of all control points affecting it. */
+	FVector2D RecomputeKeyValue(const FLatticeKeyCoords& InPointCoords, const FCellVertexIndices& InIndices) const;
+	/** Computes the weights for InKeyData such that only the corners of the cell it is in affect it. */
+	FLatticeKeyCoords ComputeSingleCellWeights(const FCellVertexIndices& InIndices, const FVector2D& InPoint) const;
+};
+}
+
diff --git a/Engine/Plugins/Editor/CurveEditorTools/Source/CurveEditorTools/Private/Tools/Lattice/Misc/LatticeDrawUtils.cpp b/Engine/Plugins/Editor/CurveEditorTools/Source/CurveEditorTools/Private/Tools/Lattice/Misc/LatticeDrawUtils.cpp
new file mode 100644
--- /dev/null
+++ b/Engine/Plugins/Editor/CurveEditorTools/Source/CurveEditorTools/Private/Tools/Lattice/Misc/LatticeDrawUtils.cpp
@@ -0,0 +1,370 @@
+﻿// Copyright Epic Games, Inc. All Rights Reserved.
+
+#include "LatticeDrawUtils.h"
+
+#include "CurveEditor.h"
+#include "CurveEditorSelection.h"
+#include "CurveModel.h"
+#include "LatticeUtils.h"
+#include "Math/NumericLimits.h"
+#include "Misc/VectorMathUtils.h"
+#include "Rendering/DrawElements.h"
+#include "SCurveEditorPanel.h"
+#include "SCurveEditorView.h"
+#include "Tools/Lattice/LatticeControlsDrawData.h"
+#include "Tools/Lattice/LatticeDeformer2D.h"
+
+namespace UE::CurveEditorTools
+{	
+FSlateLayoutTransform ComputeViewToViewContainerTransform(const SCurveEditorView& InView, const FCurveEditor& InCurveEditor)
+{
+	const FSlateLayoutTransform ViewToAbsolute = InView.GetCachedGeometry().GetAccumulatedLayoutTransform();
+	const FSlateLayoutTransform ViewContainerToAbsolute = InCurveEditor.GetPanel()->GetViewContainerGeometry().GetAccumulatedLayoutTransform();
+	const FSlateLayoutTransform AbsoluteToViewContainer = ViewContainerToAbsolute.Inverse();
+	const FSlateLayoutTransform ViewToViewContainerTransform = Concatenate(ViewToAbsolute, AbsoluteToViewContainer);
+	return ViewToViewContainerTransform;
+}
+
+FLatticeBounds ComputeBounds(const FCurveEditor& InCurveEditor)
+{
+	FLatticeBounds Bounds;
+	Bounds.MinValues = { TNumericLimits<double>::Max(), TNumericLimits<double>::Max() };
+	Bounds.MaxValues = { TNumericLimits<double>::Lowest(), TNumericLimits<double>::Lowest() };
+	Bounds.MinSlatePosition = { TNumericLimits<double>::Max(), TNumericLimits<double>::Max() };
+	Bounds.MaxSlatePosition = { TNumericLimits<double>::Lowest(), TNumericLimits<double>::Lowest() };
+	
+	const TMap<FCurveModelID, FKeyHandleSet>& SelectedKeySet = InCurveEditor.Selection.GetAll();
+	for (const TPair<FCurveModelID, FKeyHandleSet>& Pair : SelectedKeySet)
+	{
+		const SCurveEditorView* View = InCurveEditor.FindFirstInteractiveView(Pair.Key);
+		if (!View)
+		{
+			continue;
+		}
+
+		// A newly created view may have a zero-size until the next tick which is a problem if
+		// we ask the View for it's curve space, so we skip over it until it has a size.
+		if(View->GetCachedGeometry().GetLocalSize() == FVector2D::ZeroVector)
+		{
+			continue;
+		}
+		
+		FCurveModel* CurveModel = InCurveEditor.FindCurve(Pair.Key);
+		check(CurveModel);
+
+		TArrayView<const FKeyHandle> KeyHandles = Pair.Value.AsArray();
+
+		TArray<FKeyPosition> KeyPositions;
+		KeyPositions.SetNumUninitialized(KeyHandles.Num());
+		CurveModel->GetKeyPositions(KeyHandles, KeyPositions);
+
+		FCurveEditorScreenSpace ViewSpace = View->GetViewSpace();
+		const FSlateLayoutTransform ViewToViewContainerTransform = ComputeViewToViewContainerTransform(*View, InCurveEditor);
+		for (int32 i = 0; i < KeyPositions.Num(); ++i)
+		{
+			const FKeyPosition& KeyPosition = KeyPositions[i];
+			const FVector2D Position_ViewSpace = FVector2D(
+				ViewSpace.SecondsToScreen(KeyPosition.InputValue), ViewSpace.ValueToScreen(KeyPosition.OutputValue)
+				);
+			const FVector2D PanelSpaceLocation = ViewToViewContainerTransform.TransformPoint(Position_ViewSpace);
+
+			Bounds.bIsVisible = true;
+			
+			Bounds.MinValues.X = FMath::Min(KeyPosition.InputValue, Bounds.MinValues.X);
+			Bounds.MinValues.Y = FMath::Min(KeyPosition.OutputValue, Bounds.MinValues.Y);
+			Bounds.MaxValues.X = FMath::Max(KeyPosition.InputValue, Bounds.MaxValues.X);
+			Bounds.MaxValues.Y = FMath::Max(KeyPosition.OutputValue, Bounds.MaxValues.Y);
+			
+			Bounds.MinSlatePosition = FVector2D::Min(PanelSpaceLocation, Bounds.MinSlatePosition);
+			Bounds.MaxSlatePosition = FVector2D::Max(PanelSpaceLocation, Bounds.MaxSlatePosition);
+		}
+	}
+	
+	return Bounds;
+}
+
+bool TransformViewToViewContainer(
+	const FCurveEditor& InCurveEditor,
+	TConstArrayView<FVector2D> InControlPoints,
+	TArrayView<FVector2D> OutTransformedToViewContainer
+	)
+{
+	check(InControlPoints.Num() <= OutTransformedToViewContainer.Num());
+	
+	const SCurveEditorView* View = [&InCurveEditor]() -> const SCurveEditorView*
+	{
+		for (const TPair<FCurveModelID, FKeyHandleSet>& Pair : InCurveEditor.Selection.GetAll())
+		{
+			if (const SCurveEditorView* View = InCurveEditor.FindFirstInteractiveView(Pair.Key))
+			{
+				return View;
+			}
+		}
+		return nullptr;
+	}();
+	// A newly created view may have a zero-size until the next tick which is a problem if
+	// we ask the View for it's curve space, so we skip over it until it has a size.
+	if (!View || View->GetCachedGeometry().GetLocalSize() == FVector2D::ZeroVector)
+	{
+		return false;
+	}
+	
+	// The widget hierarchy is SCurveEditorViewContainer -> SCurveEditorView.
+	// ControlPoints are in ViewSpace(SCurveEditorView) and transformed to SCurveEditorViewContainer's FGeometry, received in OnPaint.
+	const FCurveEditorScreenSpace ViewSpace = View->GetViewSpace();
+	const FSlateLayoutTransform ViewToViewContainerTransform = ComputeViewToViewContainerTransform(*View, InCurveEditor);
+
+	for (int32 Index = 0; Index < InControlPoints.Num(); ++Index)
+	{
+		const FVector2D& ControlPoint = InControlPoints[Index];
+		const FVector2D Position_ViewSpace = FVector2D(
+			ViewSpace.SecondsToScreen(ControlPoint.X), ViewSpace.ValueToScreen(ControlPoint.Y)
+			);
+		const FVector2D Position_ViewContainerSpace = ViewToViewContainerTransform.TransformPoint(Position_ViewSpace);
+		OutTransformedToViewContainer[Index] = Position_ViewContainerSpace;
+	}
+	return true;
+}
+
+namespace LatticeDrawConstants
+{
+constexpr float ControlPointAnchorWidth = 13.f;
+constexpr float ControlPointHighlightAlpha = 0.15f;
+constexpr float ControlEdgeDashLength = 3.f;
+	
+constexpr float EdgeHoveredSideSize = 10.f;
+constexpr float EdgeHighlightAlpha = 0.15f;
+	
+constexpr float CellHoverInset = EdgeHoveredSideSize / 2.f;
+
+static FGeometry GetPointGeometry(const FGeometry& InLatticeGeometry, const FVector2D& ControlPoint)
+{
+	const FVector2D PointBoxSize(ControlPointAnchorWidth, ControlPointAnchorWidth);
+	const FVector2D Translation = ControlPoint - FVector2D(ControlPointAnchorWidth / 2.f);
+	const FGeometry PointGeometry = InLatticeGeometry.MakeChild(PointBoxSize, FSlateLayoutTransform(Translation));
+	return PointGeometry;
+}
+
+struct FRotatedEdgeInfo
+{
+	/** Geometry of the hovered edge rect (unrotated). */
+	FGeometry EdgeRectGeometry;
+	/** How much to turn EdgeRectGeometry (clockwise - just plug it into Slate's rotation transform as is). */
+	float AngleRadians;
+	FVector2D Offset;
+};
+
+static FRotatedEdgeInfo GetEdgeRenderInfo(const FGeometry& InViewContainerGeometry, const FVector2D& Start, const FVector2D& End, bool bOffsetEdgeCenter = true)
+{
+	const FVector2D StartToEnd = End - Start;
+	const float EdgeLength = StartToEnd.Length();
+	const FVector2D EdgeBoxSize(EdgeLength, LatticeDrawConstants::EdgeHoveredSideSize);
+	
+	// Compute angle between positive X-Axis: A*B=|A||B|cos(a). A = StartToEnd, B = X-Axis, |B|=1. A*B = A.X (as B.Y = 0).
+	// Reminder: pos. X -> Angle in [0,90], neg. X, angle in [90,180].
+	float Angle = FMath::Acos(StartToEnd.X / EdgeLength);
+	// Reminder: Slate's pos. Y-Axis points down, and pos. x-axis to the right. 
+	const bool bEdgePointsUp = StartToEnd.Y < 0.f;
+	// Mirror ACos if the point is on the neg. side of the y-axis.
+	Angle *= bEdgePointsUp ? -1.f : 1.f;
+		
+	FVector2D Translation = Start;
+	// Offset the box's side to be in the center of the control point.
+	FVector2D Offset(0.0, LatticeDrawConstants::EdgeHoveredSideSize / 2.0);
+	Offset = Offset.GetRotated(FMath::RadiansToDegrees(Angle));
+	if (bOffsetEdgeCenter)
+	{
+		Translation -= Offset;
+	}
+		
+	const FGeometry ChildGeometry = InViewContainerGeometry.MakeChild(EdgeBoxSize, FSlateLayoutTransform(Translation));
+	return { ChildGeometry, Angle, Offset};
+}
+}
+	
+FLatticeHoverState ComputeLatticeHoverState(
+	const FGeometry& InViewContainerGeometry,
+	const FVector2D& InMouseScreenPosition,
+	TConstArrayView<FVector2D> InControlPoints,
+	TConstArrayView<FLatticeControlEdge> InControlEdges,
+	int32 InNumPointsInX,
+	int32 InNumCells
+	)
+{
+	FLatticeHoverState HoverState;
+	
+	for (int32 Index = 0; Index < InControlPoints.Num(); ++Index)
+	{
+		const FVector2D& ControlPoint = InControlPoints[Index];
+		const FGeometry PointGeometry = LatticeDrawConstants::GetPointGeometry(InViewContainerGeometry, ControlPoint);
+		if (PointGeometry.IsUnderLocation(InMouseScreenPosition))
+		{
+			HoverState.HoveredControlPoint = Index;
+			// Don't hover anything else.
+			return HoverState;
+		}
+	}
+
+	const FSlateRenderTransform& ViewContainerToAbsolute = InViewContainerGeometry.GetAccumulatedRenderTransform();
+	const FSlateRenderTransform& AbsoluteToViewContainer = ViewContainerToAbsolute.Inverse();
+	const FVector2D Mouse_ViewContainer = AbsoluteToViewContainer.TransformPoint(InMouseScreenPosition);
+	for (int32 Index = 0; Index < InControlEdges.Num(); ++Index)
+	{
+		const FLatticeControlEdge& Edge = InControlEdges[Index];
+		const FVector2D& Start = Edge.Start();
+		const FVector2D& End = Edge.End();
+		constexpr bool bOffsetEdgeCenter = false; // We need to do the offsetting ourselves below.
+		const auto[EdgeGeometry, Angle, Offset] = LatticeDrawConstants::GetEdgeRenderInfo(InViewContainerGeometry, Start, End, bOffsetEdgeCenter);
+
+		// We're going to compute the direction vector from start to mouse, rotate it instead of the box, and check whether the adjust mouse position is in the box.
+		const FVector2D StartCenter_ViewContainer = Start - Offset;
+		const FVector2D StartCenterToMouse_ViewContainer = Mouse_ViewContainer - StartCenter_ViewContainer;
+		const FVector2D RotatedStartToMouse_ViewContainer = StartCenterToMouse_ViewContainer.GetRotated(FMath::RadiansToDegrees(-Angle));
+		const FVector2D RotatedMouse_ViewContainer = Start + RotatedStartToMouse_ViewContainer;
+		const FSlateRect BoundingBox_ScreenSpace = EdgeGeometry.GetLayoutBoundingRect();
+		const FVector2D RotatedMousePosition_ScreenSpace = ViewContainerToAbsolute.TransformPoint(RotatedMouse_ViewContainer);
+		if (BoundingBox_ScreenSpace.ContainsPoint(RotatedMousePosition_ScreenSpace))
+		{
+			HoverState.HoveredEdge = Index;
+			return HoverState;
+		}
+	}
+
+	for (int32 CellIndex = 0; CellIndex < InNumCells; ++CellIndex)
+	{
+		const auto[TopLeft, TopRight, BottomRight, BottomLeft] = GetCellIndices(CellIndex, InNumPointsInX);
+		const bool bIsInTriangle1 = IsPointInTriangle(
+			Mouse_ViewContainer, InControlPoints[TopLeft], InControlPoints[TopRight], InControlPoints[BottomRight]
+			);
+		const bool bIsInTriangle2 = IsPointInTriangle(
+			Mouse_ViewContainer, InControlPoints[BottomRight], InControlPoints[BottomLeft], InControlPoints[TopLeft]
+			);
+		const bool bIsInCell = bIsInTriangle1 || bIsInTriangle2;
+		if (bIsInCell)
+		{
+			HoverState.HoveredCell = CellIndex;
+			return HoverState;
+		}
+	}
+
+	return HoverState;
+}
+
+static void DrawControlPoints(
+	const FLatticeControlsDrawData& InDrawData, const FGeometry& InViewContainerGeometry, FSlateWindowElementList& OutDrawElements, int32 InPaintOnLayerId
+	)
+{
+	for (int32 Index = 0; Index < InDrawData.ControlPoints.Num(); ++Index)
+	{
+		const FVector2D& ControlPoint = InDrawData.ControlPoints[Index];
+		const FGeometry PointGeometry = LatticeDrawConstants::GetPointGeometry(InViewContainerGeometry, ControlPoint);
+		FSlateDrawElement::MakeBox(
+			OutDrawElements, InPaintOnLayerId, PointGeometry.ToPaintGeometry(), FAppStyle::GetBrush(TEXT("MarqueeSelection"))
+			);
+		
+		if (InDrawData.HoverState.HoveredControlPoint.Get(INDEX_NONE) == Index)
+		{
+			const FLinearColor HighlightColor = FLinearColor::White.CopyWithNewOpacity(LatticeDrawConstants::ControlPointHighlightAlpha);
+			FSlateDrawElement::MakeBox(
+				OutDrawElements, InPaintOnLayerId, PointGeometry.ToPaintGeometry(), FAppStyle::GetBrush(TEXT("WhiteBrush")),
+				ESlateDrawEffect::None, HighlightColor
+			);
+		}
+	}
+}
+
+static void DrawControlEdges(
+	const FLatticeControlsDrawData& InDrawData, const FGeometry& InViewContainerGeometry, FSlateWindowElementList& OutDrawElements, int32 InPaintOnLayerId
+	)
+{
+	for (int32 Index = 0; Index < InDrawData.ControlEdges.Num(); ++Index)
+	{
+		const FLatticeControlEdge& Edge = InDrawData.ControlEdges[Index];
+		const FVector2D& Start = Edge.Start();
+		const FVector2D& End = Edge.End();
+		
+		FSlateDrawElement::MakeDashedLines(
+			OutDrawElements,
+			InPaintOnLayerId,
+			InViewContainerGeometry.ToPaintGeometry(),
+			{ FVector2f(Start.X, Start.Y), FVector2f(End.X, End.Y) },
+			ESlateDrawEffect::None,
+			FLinearColor::White,
+			1.f,
+			LatticeDrawConstants::ControlEdgeDashLength
+			);
+
+		if (InDrawData.HoverState.HoveredEdge == Index)
+		{
+			const auto[EdgeGeometry, Angle, _] = LatticeDrawConstants::GetEdgeRenderInfo(InViewContainerGeometry, Start, End);
+			FSlateDrawElement::MakeRotatedBox(
+				OutDrawElements,
+				InPaintOnLayerId,
+				EdgeGeometry.ToPaintGeometry(),
+				FAppStyle::GetBrush(TEXT("WhiteBrush")),
+				ESlateDrawEffect::None,
+				Angle,
+				FVector2f(0, 0),
+				FSlateDrawElement::RelativeToElement,
+				FLinearColor::White.CopyWithNewOpacity(LatticeDrawConstants::EdgeHighlightAlpha)
+				);
+		}
+	}
+}
+
+static void DrawHoveredCells(
+	const FLatticeControlsDrawData& InDrawData, const FGeometry& InViewContainerGeometry, FSlateWindowElementList& OutDrawElements, int32 InPaintOnLayerId
+	)
+{
+	if (!InDrawData.HoverState.HoveredCell)
+	{
+		return;
+	}
+
+	const auto[TopLeft, TopRight, BottomRight, BottomLeft] = GetCellIndices(*InDrawData.HoverState.HoveredCell, InDrawData.MatrixWidth);
+	FVector2D Vert_TopLeft = InDrawData.ControlPoints[TopLeft];
+	FVector2D Vert_TopRight = InDrawData.ControlPoints[TopRight];
+	FVector2D Vert_BottomRight = InDrawData.ControlPoints[BottomRight];
+	FVector2D Vert_BottomLeft = InDrawData.ControlPoints[BottomLeft];
+	InsetQuadBy(Vert_TopLeft, Vert_TopRight, Vert_BottomRight, Vert_BottomLeft, LatticeDrawConstants::CellHoverInset);
+	
+	const FVector2f TexturePos(0.f, 0.f);
+	const FSlateRenderTransform& Transform = InViewContainerGeometry.GetAccumulatedRenderTransform();
+	const FColor TintColor = FLinearColor::White.CopyWithNewOpacity(LatticeDrawConstants::EdgeHighlightAlpha).ToFColor(true);
+	const TArray Vertices {
+		FSlateVertex::Make<ESlateVertexRounding::Disabled>(Transform, FVector2f(Vert_TopLeft.X, Vert_TopLeft.Y), TexturePos, TintColor),
+		FSlateVertex::Make<ESlateVertexRounding::Disabled>(Transform, FVector2f(Vert_TopRight.X, Vert_TopRight.Y), TexturePos, TintColor),
+		FSlateVertex::Make<ESlateVertexRounding::Disabled>(Transform, FVector2f(Vert_BottomRight.X, Vert_BottomRight.Y), TexturePos, TintColor),
+		FSlateVertex::Make<ESlateVertexRounding::Disabled>(Transform, FVector2f(Vert_BottomLeft.X, Vert_BottomLeft.Y), TexturePos, TintColor)
+	};
+	
+	// Handle concave quad: If the quad is concave, then the diagonal of the 2 triangles must not be on the vertex that has the concave angle.
+	// Note: this does not handle crossed quads (i.e. where the "quad" is 2 triangles with 2 "opposite" edges crossing each other).
+	const float AngleTopLeft = FMath::Acos(FVector2D::DotProduct(Vert_BottomLeft - Vert_TopLeft, Vert_TopRight - Vert_TopLeft));
+	const float AngleBottomRight = FMath::Acos(FVector2D::DotProduct(Vert_BottomLeft - Vert_BottomRight, Vert_TopRight - Vert_BottomRight));
+	const bool bRotate = AngleTopLeft >= FMath::DegreesToRadians(180.f) || AngleBottomRight >= FMath::DegreesToRadians(180.f);
+	const TArray<SlateIndex> Indices = bRotate ? TArray<SlateIndex>{ 0, 1, 2, 0, 2, 3 } : TArray<SlateIndex>{ 1, 2, 3, 3, 0, 1};
+
+	const FSlateResourceHandle ResourceHandle = FSlateApplication::Get().GetRenderer()->GetResourceHandle(
+		*FAppStyle::GetBrush("Sequencer.SectionArea.Background") // This is completely white so we can just multiply it with the TintColor
+		);
+	FSlateDrawElement::MakeCustomVerts(
+		OutDrawElements, InPaintOnLayerId, ResourceHandle, Vertices, Indices, nullptr, 0, 0
+		);
+}
+
+void DrawLatticeControls(
+	const FLatticeControlsDrawData& InDrawData,
+	const FGeometry& InViewContainerGeometry,
+	FSlateWindowElementList& OutDrawElements,
+	int32 InPaintOnLayerId
+	)
+{
+	// Cells highlight first...
+	DrawHoveredCells(InDrawData, InViewContainerGeometry, OutDrawElements, InPaintOnLayerId);
+	// ... as edges are drawn over cells...
+	DrawControlEdges(InDrawData, InViewContainerGeometry, OutDrawElements, InPaintOnLayerId);
+	// ... and control points draw over edges.
+	DrawControlPoints(InDrawData, InViewContainerGeometry, OutDrawElements, InPaintOnLayerId);
+}
+}
diff --git a/Engine/Plugins/Editor/CurveEditorTools/Source/CurveEditorTools/Private/Tools/Lattice/Misc/LatticeDrawUtils.h b/Engine/Plugins/Editor/CurveEditorTools/Source/CurveEditorTools/Private/Tools/Lattice/Misc/LatticeDrawUtils.h
new file mode 100644
--- /dev/null
+++ b/Engine/Plugins/Editor/CurveEditorTools/Source/CurveEditorTools/Private/Tools/Lattice/Misc/LatticeDrawUtils.h
@@ -0,0 +1,89 @@
+﻿// Copyright Epic Games, Inc. All Rights Reserved.
+
+#pragma once
+
+#include "Containers/ContainersFwd.h"
+#include "HAL/Platform.h"
+#include "Math/Vector2D.h"
+
+class FCurveEditor;
+class FSlateLayoutTransform;
+class FSlateWindowElementList;
+class SCurveEditorView;
+struct FGeometry;
+
+/** Put utils specific to drawing the lattice controls into this file */
+namespace UE::CurveEditorTools
+{
+struct FLatticeControlEdge;
+struct FLatticeControlsDrawData;
+struct FLatticeHoverState;
+
+/**
+ * Obtains the transform required to transform evaluated key positions to that space used by ICurveEditorToolExtension::OnPaint.
+ * The widget hierarchy is SCurveEditorViewContainer -> SCurveEditorView.
+ * - SCurveEditorView is used to evaluate key positions.
+ * - ICurveEditorToolExtension::OnPaint receives the FGeometry of SCurveEditorView
+ */
+FSlateLayoutTransform ComputeViewToViewContainerTransform(const SCurveEditorView& InView, const FCurveEditor& InCurveEditor);
+	
+struct FLatticeBounds
+{
+	/** If false, then do not discard the other values. */
+	bool bIsVisible = false;
+
+	// Values in graph
+	FVector2D MinValues;
+	FVector2D MaxValues;
+
+	// Values on screen
+	FVector2D MinSlatePosition;
+	FVector2D MaxSlatePosition;
+};
+/** Computes the min and max slate and value points in the user's selection. */
+FLatticeBounds ComputeBounds(const FCurveEditor& InCurveEditor);
+
+/**
+ * Transforms the values of the control points to where they should be drawn in InLatticeGeometry.
+ * @param InCurveEditor Needed to get interactive view for transforming values.
+ * @param InBounds The bounds as computed by ComputeBounds
+ * @param InLatticeGeometry The result of MakeLatticeGeometry
+ * @param InControlPoints The points to transform
+ * @param OutTransformedToViewContainer Will contain the transformed points; must be equal length to InControlPoints.
+ * @return Whether OutTransformed was written to.
+ */
+bool TransformViewToViewContainer(
+	const FCurveEditor& InCurveEditor,
+	TConstArrayView<FVector2D> InControlPoints,
+	TArrayView<FVector2D> OutTransformedToViewContainer
+	);
+
+/**
+ * Computes which elements should be hovered.
+ * @param InViewContainerGeometry The result of MakeLatticeGeometry
+ * @param InMouseScreenPosition The screen space position of the mouse
+ * @param InControlPoints The control points of the lattice deformer, already transformed using TransformControlPointsToSlate.
+ * @param InControlEdges The edges of the lattice deformer.
+ * @param InNumPointsInX The number of points InControlPoints has in x direction
+ * @param InNumCells Total number of cells in the deformer
+ */
+FLatticeHoverState ComputeLatticeHoverState(
+	const FGeometry& InViewContainerGeometry,
+	const FVector2D& InMouseScreenPosition,
+	TConstArrayView<FVector2D> InControlPoints,
+	TConstArrayView<FLatticeControlEdge> InControlEdges,
+	int32 InNumPointsInX,
+	int32 InNumCells
+	);
+
+/**
+ * Draws the lattice controls.
+ * @param InDrawData 
+ * @param InViewContainerGeometry The result of MakeLatticeGeometry
+ * @param OutDrawElements 
+ * @param InPaintOnLayerId 
+ */
+void DrawLatticeControls(
+	const FLatticeControlsDrawData& InDrawData, const FGeometry& InViewContainerGeometry, FSlateWindowElementList& OutDrawElements, int32 InPaintOnLayerId
+	);
+}
diff --git a/Engine/Plugins/Editor/CurveEditorTools/Source/CurveEditorTools/Private/Tools/Lattice/Misc/LatticeUtils.cpp b/Engine/Plugins/Editor/CurveEditorTools/Source/CurveEditorTools/Private/Tools/Lattice/Misc/LatticeUtils.cpp
new file mode 100644
--- /dev/null
+++ b/Engine/Plugins/Editor/CurveEditorTools/Source/CurveEditorTools/Private/Tools/Lattice/Misc/LatticeUtils.cpp
@@ -0,0 +1,171 @@
+﻿// Copyright Epic Games, Inc. All Rights Reserved.
+
+#include "LatticeUtils.h"
+
+#include "CurveEditor.h"
+#include "CurveEditorSelection.h"
+#include "CurveEditorTypes.h"
+#include "Math/NumericLimits.h"
+#include "Tools/Lattice/LatticeDeformer2D.h"
+
+#include <type_traits>
+
+namespace UE::CurveEditorTools
+{
+int32 MatrixIndicesToFlatIndex(int32 InColumn, int32 InRow, int32 InNumColumns)
+{
+	return InColumn + InRow * InNumColumns;
+}
+
+FMatrixIndices FlatIndexToMatrixIndices(int32 InFlatIndex, int32 InNumColumns)
+{
+	const int32 Col = InFlatIndex % InNumColumns;
+	const int32 Row = InFlatIndex / InNumColumns;
+	return FMatrixIndices{ Col, Row };
+}
+
+FCellCornerIndices GetMatrixCellSlotsWithCorner(int32 InCornerIndex, int32 InNumLatticePointColumns, int32 InNumLatticePointRows)
+{
+	const auto[Column, Row] = FlatIndexToMatrixIndices(InCornerIndex, InNumLatticePointColumns);
+
+	FCellCornerIndices Result;
+	const int32 MaxColumnIndex = InNumLatticePointColumns - 1;
+	const int32 MaxRowIndex = InNumLatticePointRows - 1;
+	
+	const bool bHasTopLeftCell = Column > 0 && Row > 0;
+	const bool bHasTopRightCell = Column < MaxColumnIndex && Row > 0;
+	const bool bHasBottomRightCell = Column < MaxColumnIndex && Row < MaxRowIndex;
+	const bool bHasBottomLeftCell = Column > 0 && Row < MaxRowIndex;
+
+	Result.TopLeft = bHasTopLeftCell			? MatrixIndicesToFlatIndex(Column - 1	, Row - 1		, InNumLatticePointColumns - 1) : INDEX_NONE;
+	Result.TopRight = bHasTopRightCell			? MatrixIndicesToFlatIndex(Column		, Row - 1		, InNumLatticePointColumns - 1) : INDEX_NONE;
+	Result.BottomLeft = bHasBottomLeftCell		? MatrixIndicesToFlatIndex(Column - 1	, Row			, InNumLatticePointColumns - 1) : INDEX_NONE;
+	Result.BottomRight = bHasBottomRightCell	? MatrixIndicesToFlatIndex(Column 		, Row			, InNumLatticePointColumns - 1) : INDEX_NONE;
+
+	return Result;
+}
+
+FEdgeVertexIndices GetEdgeIndices(int32 InEdgeIndex, int32 InNumPointsInWidth, int32 InNumPointsInHeight)
+{
+	const int32 HorizontalEdges = (InNumPointsInWidth - 1) * InNumPointsInHeight;
+
+	// Determine if the edge is horizontal
+	if (InEdgeIndex < HorizontalEdges)
+	{
+		// Horizontal edge
+		const int32 Row = InEdgeIndex / (InNumPointsInWidth - 1);            
+		const int32 Col = InEdgeIndex % (InNumPointsInWidth - 1);          
+		const int32 Left = Row * InNumPointsInWidth + Col;              
+		const int32 Right = Row * InNumPointsInWidth + (Col + 1);        
+		return { Left, Right };
+	}
+	else
+	{
+		// Vertical edge
+		const int32 VerticalIndex = InEdgeIndex - HorizontalEdges;
+		const int32 Column = VerticalIndex % InNumPointsInWidth;             
+		const int32 Row = VerticalIndex / InNumPointsInWidth;             
+		const int32 Top = Row * InNumPointsInWidth + Column;              
+		const int32 Bottom = (Row + 1) * InNumPointsInWidth + Column;        
+		return { Top, Bottom };
+	}
+}
+
+FCellVertexIndices GetCellIndices(int32 InCellIndex, int32 InNumPointsInWidth)
+{
+	FCellVertexIndices CellIndices;
+	
+	const int32 CellWidth = InNumPointsInWidth - 1;
+	const auto[CellColumn, CellRow] = FlatIndexToMatrixIndices(InCellIndex, CellWidth);
+	
+	CellIndices.TopLeft = MatrixIndicesToFlatIndex(CellColumn, CellRow, InNumPointsInWidth);
+	CellIndices.TopRight = MatrixIndicesToFlatIndex(CellColumn + 1, CellRow, InNumPointsInWidth);
+	CellIndices.BottomRight = MatrixIndicesToFlatIndex(CellColumn + 1, CellRow + 1, InNumPointsInWidth);
+	CellIndices.BottomLeft = MatrixIndicesToFlatIndex(CellColumn, CellRow + 1, InNumPointsInWidth);
+
+	return CellIndices;
+}
+
+template<typename TCallback> requires std::is_invocable_v<TCallback, const FKeyPosition&>
+static void ForEachKey(const FCurveEditor& InCurveEditor, TCallback&& Callback)
+{
+	const TMap<FCurveModelID, FKeyHandleSet>& Selection = InCurveEditor.Selection.GetAll();
+	for (const TPair<FCurveModelID, FKeyHandleSet>& Pair : Selection)
+	{
+		FCurveModel* CurveModel = InCurveEditor.FindCurve(Pair.Key);
+		if (!CurveModel)
+		{
+			return;
+		}
+
+		const TConstArrayView<FKeyHandle> KeyHandles = Pair.Value.AsArray();
+		TArray<FKeyPosition> KeyPositions;
+		KeyPositions.SetNumUninitialized(KeyHandles.Num());
+		CurveModel->GetKeyPositions(KeyHandles, KeyPositions);
+		
+		for (const FKeyPosition& Position : KeyPositions)
+		{
+			Callback(Position);
+		}
+	}
+}
+
+TArray<FVector2D> PlaceLatticeGridOverNewKeyPositions(const FLatticeDeformer2D& InDeformer, const FCurveEditor& InCurveEditor)
+{
+	checkfSlow(!InCurveEditor.Selection.GetAll().IsEmpty(), TEXT("The key selection is supposed to be the same as when the deformer was created!")); 
+	FVector2D NewMinValues { TNumericLimits<double>::Max(), TNumericLimits<double>::Max() };;
+	ForEachKey(InCurveEditor, [&NewMinValues](const FKeyPosition& Position)
+	{
+		NewMinValues.X = FMath::Min(Position.InputValue, NewMinValues.X);
+		NewMinValues.Y = FMath::Min(Position.OutputValue, NewMinValues.Y);
+	});
+	
+	FVector2D OldMinValues { TNumericLimits<double>::Max(), TNumericLimits<double>::Max() };
+	for (const FVector2D& ControlPoint : InDeformer.GetControlPoints())
+	{
+		OldMinValues.X = FMath::Min(ControlPoint.X, OldMinValues.X);
+		OldMinValues.Y = FMath::Min(ControlPoint.Y, OldMinValues.Y);
+	}
+
+	// The algorithm assumes that all points have been translated uniformly, like you would if you move the key selection with middle mouse button.
+	const FVector2D OldMinToNewMin = NewMinValues - OldMinValues;
+	
+	TArray<FVector2D> TransformedControlPoints;
+	TransformedControlPoints.SetNumUninitialized(InDeformer.NumControlPoints());
+	for (int32 ControlPointIndex = 0; ControlPointIndex < InDeformer.NumControlPoints(); ++ControlPointIndex)
+	{
+		TransformedControlPoints[ControlPointIndex] = InDeformer.GetControlPoints()[ControlPointIndex] + OldMinToNewMin;
+	}
+	
+	return TransformedControlPoints;
+}
+
+void AddSelectionToDeformer(FLatticeDeformer2D& InDeformer, const FCurveEditor& InCurveEditor)
+{
+	for (const TPair<FCurveModelID, FKeyHandleSet>& Pair : InCurveEditor.Selection.GetAll())
+	{
+		const FCurveModelID CurveId = Pair.Key;
+		const FCurveModel* CurveModel = InCurveEditor.FindCurve(CurveId);
+		if (!CurveModel)
+		{
+			continue;
+		}
+
+		// Looping like this avoids allocating a big FKeyPosition array.
+		TConstArrayView<FKeyHandle> Keys = Pair.Value.AsArray();
+		for (int32 Index = 0; Index < Keys.Num(); ++Index)
+		{
+			const FKeyHandle& KeyId = Keys[Index];
+				
+			FKeyPosition Position;
+			TConstArrayView<FKeyHandle> SingleKeyView(&KeyId, 1);
+			TArrayView<FKeyPosition> SinglePositionView(&Position, 1);
+			CurveModel->GetKeyPositions(SingleKeyView, SinglePositionView);
+
+			const FVector2D AsVector{ Position.InputValue, Position.OutputValue };
+			const TConstArrayView<FVector2D> KeyView(&AsVector, 1);
+			InDeformer.AddPoints_BeforeLatticeMoved({ FKeyInCurveId{ KeyId, CurveId } }, KeyView);
+		}
+	}
+}
+}
diff --git a/Engine/Plugins/Editor/CurveEditorTools/Source/CurveEditorTools/Private/Tools/Lattice/Misc/LatticeUtils.h b/Engine/Plugins/Editor/CurveEditorTools/Source/CurveEditorTools/Private/Tools/Lattice/Misc/LatticeUtils.h
new file mode 100644
--- /dev/null
+++ b/Engine/Plugins/Editor/CurveEditorTools/Source/CurveEditorTools/Private/Tools/Lattice/Misc/LatticeUtils.h
@@ -0,0 +1,80 @@
+﻿// Copyright Epic Games, Inc. All Rights Reserved.
+
+#pragma once
+
+#include "Misc/AssertionMacros.h"
+#include "Misc/CoreMiscDefines.h"
+#include "HAL/Platform.h"
+
+class FCurveEditor;
+template<typename OptionalType> struct TOptional;
+
+/** Put utils specific to lattice deforming into this file */
+namespace UE::CurveEditorTools
+{
+class FLatticeDeformer2D;
+	
+/** @return Index you can use in a TArray representing a 2D matrix. */
+int32 MatrixIndicesToFlatIndex(int32 InColumn, int32 InRow, int32 InNumColumns);
+
+struct FMatrixIndices
+{
+	int32 Column; // X
+	int32 Row; // Y
+};
+/** @return Matrix Column and Row that InFlatIndex maps to if matrix is InNumColumns wide. Does not check bounds.*/
+FMatrixIndices FlatIndexToMatrixIndices(int32 InFlatIndex, int32 InNumColumns);
+
+struct FCellCornerIndices
+{
+	int32 TopLeft;
+	int32 TopRight;
+	int32 BottomRight;
+	int32 BottomLeft;
+};
+/** Suppose you have nxm matrix of lattice points that forms a (n-1)x(m-1) matrix of cells, returns the flat indices of the cells that share InCornerIndex. */
+FCellCornerIndices GetMatrixCellSlotsWithCorner(int32 InCornerIndex, int32 InNumLatticePointColumns, int32 InNumLatticePointRows);
+
+struct FEdgeVertexIndices
+{
+	int32 Start = INDEX_NONE;
+	int32 End = INDEX_NONE;
+};
+/**
+ * @param InEdgeIndex Index of the edge according to how it is described in FLatticeDeformer2D.
+ * @param InNumPointsInWidth The number of points in the x direction (num cells in x-direction + 1)
+ * @param InNumPointsInHeight The number of points in the z direction (num cells in z-direction + 1)
+ * @return The corner indices of the cell.
+ */
+FEdgeVertexIndices GetEdgeIndices(int32 InEdgeIndex, int32 InNumPointsInWidth, int32 InNumPointsInHeight);
+
+struct FCellVertexIndices
+{
+	int32 TopLeft = INDEX_NONE;
+	int32 TopRight = INDEX_NONE;
+	int32 BottomRight = INDEX_NONE;
+	int32 BottomLeft = INDEX_NONE;
+
+	bool IsValid() const
+	{
+		// Either all or none are INDEX_NONE
+		return TopLeft != INDEX_NONE && ensure(TopRight != INDEX_NONE) && ensure(BottomRight != INDEX_NONE) && ensure(BottomLeft != INDEX_NONE);
+	}
+};
+/**
+ * @param InCellIndex Index of the cell according to how it is described in FLatticeDeformer2D.
+ * @param InNumPointsInWidth The number of points in the x direction (num cells in x-direction + 1)
+ * @return The corner indices of the cell.
+ */
+FCellVertexIndices GetCellIndices(int32 InCellIndex, int32 InNumPointsInWidth);
+
+/**
+ * Call this in response to the positions of the key selection having been moved.
+ * Assuming that all points have been translated uniformly, this moves the the lattice grid by the delta movement.
+ * @return The new control points that InDeformer should have, or empty, if the lattice grid needs to be regenerated.
+ */
+TArray<FVector2D> PlaceLatticeGridOverNewKeyPositions(const FLatticeDeformer2D& InDeformer, const FCurveEditor& InCurveEditor);
+
+/** Adds InCurveEditor's selection to InDeformer. */
+void AddSelectionToDeformer(FLatticeDeformer2D& InDeformer, const FCurveEditor& InCurveEditor);
+}
diff --git a/Engine/Plugins/Editor/CurveEditorTools/Source/CurveEditorTools/Private/Tools/Lattice/UObject/LatticeUndoData.cpp b/Engine/Plugins/Editor/CurveEditorTools/Source/CurveEditorTools/Private/Tools/Lattice/UObject/LatticeUndoData.cpp
new file mode 100644
--- /dev/null
+++ b/Engine/Plugins/Editor/CurveEditorTools/Source/CurveEditorTools/Private/Tools/Lattice/UObject/LatticeUndoData.cpp
@@ -0,0 +1,68 @@
+﻿// Copyright Epic Games, Inc. All Rights Reserved.
+
+#include "LatticeUndoData.h"
+
+#include "CurveEditor.h"
+#include "Tools/Lattice/LatticeDeformer2D.h"
+
+void UCurveEditorTools_LatticeUndoData::SaveState(const FCurveEditor& InEditor, const UE::CurveEditorTools::FLatticeDeformer2D& InDeformer)
+{
+	ControlPoints = InDeformer.GetControlPoints();
+
+	// The user might change selection - when restoring, we must restore the keys
+	CellData.Reset(InDeformer.NumCells());
+	for (int32 CellIndex = 0; CellIndex < InDeformer.NumCells(); ++CellIndex)
+	{
+		CellData.Emplace(
+			InDeformer.GetKeyCoordsInCell(CellIndex),
+			InDeformer.GetKeyIdsInCell(CellIndex)
+			);
+	}
+	
+	Selection = InEditor.Selection;
+}
+
+void UCurveEditorTools_LatticeUndoData::RestoreDeformerState(UE::CurveEditorTools::FLatticeDeformer2D& InDeformer) const
+{
+	// ControlPoints is empty if the user undoes selecting the first selection that created the first deformer.
+	if (ControlPoints.IsEmpty()
+		// In the future, our tool may be extended to change the number of lattice control points in a UI.
+		// Warn implementing developers in case they forget to update this location.
+		|| !ensureAlways(ControlPoints.Num() == InDeformer.NumControlPoints()))
+	{
+		return;
+	}
+	
+	for (int32 PointIndex = 0; PointIndex < ControlPoints.Num(); ++PointIndex)
+	{
+		TConstArrayView<FVector2D> ControlPoint(&ControlPoints[PointIndex], 1);
+		InDeformer.SetControlsPointNoRecompute({ PointIndex }, { ControlPoint });
+	}
+	
+	for (int32 CellIndex = 0; CellIndex < CellData.Num(); ++CellIndex)
+	{
+		const FCellData& Cell = CellData[CellIndex];
+		InDeformer.SetKeyIdsInCell(CellIndex, Cell.KeyIds);
+		InDeformer.SetKeyCoordsInCell(CellIndex, Cell.KeyCoords);
+	}
+}
+
+void UCurveEditorTools_LatticeUndoData::RestoreEditorState(FCurveEditor& InEditor)
+{
+	InEditor.Selection = Selection;
+}
+
+void UCurveEditorTools_LatticeUndoData::Serialize(FArchive& Ar)
+{
+	Super::Serialize(Ar);
+
+	Selection.SerializeSelection(Ar);
+	Ar << CellData;
+}
+
+void UCurveEditorTools_LatticeUndoData::PostEditUndo()
+{
+	UObject::PostEditUndo();
+	
+	OnPostEditUndo.Broadcast();
+}
diff --git a/Engine/Plugins/Editor/CurveEditorTools/Source/CurveEditorTools/Private/Tools/Lattice/UObject/LatticeUndoData.h b/Engine/Plugins/Editor/CurveEditorTools/Source/CurveEditorTools/Private/Tools/Lattice/UObject/LatticeUndoData.h
new file mode 100644
--- /dev/null
+++ b/Engine/Plugins/Editor/CurveEditorTools/Source/CurveEditorTools/Private/Tools/Lattice/UObject/LatticeUndoData.h
@@ -0,0 +1,91 @@
+﻿// Copyright Epic Games, Inc. All Rights Reserved.
+
+#pragma once
+
+#include "CurveEditorSelection.h"
+#include "Delegates/Delegate.h"
+#include "Serialization/Archive.h"
+#include "Tools/Lattice/LatticeDeformer2D.h"
+#include "UObject/Object.h"
+#include "LatticeUndoData.generated.h"
+
+class FCurveEditor;
+namespace UE::CurveEditorTools
+{
+	class FLatticeDeformer2D;
+
+	enum class ERestoreResult
+	{
+		Restored,
+		DestroyDeformer
+	};
+}
+
+/** Stores data for the purposes of undo / redo in the lattice tool. */
+UCLASS()
+class UCurveEditorTools_LatticeUndoData : public UObject
+{
+	GENERATED_BODY()
+public:
+
+	DECLARE_MULTICAST_DELEGATE(FOnPostEditUndo);
+	/** Invoked when the user has performed an undo / redo operation. You should call RestoreState. */
+	FOnPostEditUndo OnPostEditUndo;
+
+	/** Saves any potential state that needs to be captured for undo and redo. */
+	void SaveState(const FCurveEditor& InEditor, const UE::CurveEditorTools::FLatticeDeformer2D& InDeformer);
+
+	
+	/** Call in response to OnPostEditUndo. */
+	void RestoreDeformerState(UE::CurveEditorTools::FLatticeDeformer2D& InDeformer) const;
+	void RestoreEditorState(FCurveEditor& InEditor);
+
+	/** @return The serial number of the selection saved */
+	uint32 GetSelectionSerialNumber() const { return Selection.GetSerialNumber(); }
+
+	//~ Begin UObject Interface
+	virtual void Serialize(FArchive& Ar) override;
+	virtual void PostEditUndo() override;
+	//~ End UObject Interface
+
+private:
+	
+	/** The points the lattice deformer had when the operation was started. */
+	UPROPERTY()
+	TArray<FVector2D> ControlPoints;
+
+	struct FCellData
+	{
+		TArray<FVector2D> KeyCoords;
+		TArray<UE::CurveEditorTools::FKeyInCurveId> KeyIds;
+
+		FCellData() = default; // Needed by TArray::operator<<
+		explicit FCellData(TConstArrayView<FVector2D> KeyCoords, TConstArrayView<UE::CurveEditorTools::FKeyInCurveId> KeyIds)
+			: KeyCoords(KeyCoords)
+			, KeyIds(KeyIds)
+		{}
+
+		friend FArchive& operator<<(FArchive& InArchive, FCellData& InCellData)
+		{
+			InArchive << InCellData.KeyCoords;
+			InArchive << InCellData.KeyIds;
+			return InArchive;
+		}
+	};
+	/**
+	 * Maps cells to the key data they contained
+	 * 
+	 * Each index in the outer array corresponds to a cell index.
+	 * Each element corresponds to the coordinates of keys in that cell.
+	 */
+	TArray<FCellData> CellData;
+
+	/**
+	 * The lattice operates on the active selection.
+	 * 
+	 * Suppose the user changes curve A, then curve B, and undoes changing B.
+	 * Lattice grid is put back to the keys of curve A... but we must also revert the selection to be what they had when the lattice was over
+	 * the keys of curve A (or the selection stays with the keys of curve B).
+	 */
+	FCurveEditorSelection Selection;
+};
diff --git a/Engine/Source/Editor/CurveEditor/Public/CurveEditorSnapMetrics.h b/Engine/Source/Editor/CurveEditor/Public/CurveEditorSnapMetrics.h
--- a/Engine/Source/Editor/CurveEditor/Public/CurveEditorSnapMetrics.h
+++ b/Engine/Source/Editor/CurveEditor/Public/CurveEditorSnapMetrics.h
@@ -30,13 +30,13 @@ struct FCurveSnapMetrics
 	FFrameRate InputSnapRate;
 
 	/** Snap the specified input time to the input snap rate if necessary */
-	FORCEINLINE double SnapInputSeconds(double InputTime)
+	FORCEINLINE double SnapInputSeconds(double InputTime) const
 	{
 		return bSnapInputValues && InputSnapRate.IsValid() ? (InputTime * InputSnapRate).RoundToFrame() / InputSnapRate : InputTime;
 	}
 	
 	/** Snap the specified output value to the output snap interval if necessary */
-	FORCEINLINE double SnapOutput(double OutputValue)
+	FORCEINLINE double SnapOutput(double OutputValue) const
 	{
 		return bSnapOutputValues && !AllGridLines.IsEmpty() ? *Algo::MinElement(AllGridLines,
 			[OutputValue](double Val1, double Val2) { return FMath::Abs(Val1 - OutputValue) < FMath::Abs(Val2 - OutputValue); }
diff --git a/Engine/Source/Editor/EditorStyle/Private/SlateEditorStyle.cpp b/Engine/Source/Editor/EditorStyle/Private/SlateEditorStyle.cpp
--- a/Engine/Source/Editor/EditorStyle/Private/SlateEditorStyle.cpp
+++ b/Engine/Source/Editor/EditorStyle/Private/SlateEditorStyle.cpp
@@ -7526,6 +7526,8 @@ void FSlateEditorStyle::FStyle::SetupToolkitStyles()
 		Set("CurveEditorTools.ActivateRetimeTool.Small", new IMAGE_BRUSH("GenericCurveEditor/Icons/ToolsRetime_48x", Icon24x24));
 		Set("CurveEditorTools.ActivateMultiScaleTool", new IMAGE_BRUSH("GenericCurveEditor/Icons/ToolsMultiScale_48x", Icon48x48));
 		Set("CurveEditorTools.ActivateMultiScaleTool.Small", new IMAGE_BRUSH("GenericCurveEditor/Icons/ToolsMultiScale_48x", Icon24x24));
+		Set("CurveEditorTools.ActivateLatticeTool", new IMAGE_BRUSH("GenericCurveEditor/Icons/ToolsMultiScale_48x", Icon48x48));
+		Set("CurveEditorTools.ActivateLatticeTool.Small", new IMAGE_BRUSH("GenericCurveEditor/Icons/ToolsMultiScale_48x", Icon24x24));
 	}
 
 	// PhysicsAssetEditor
diff --git a/Engine/Source/Editor/EditorStyle/Private/StarshipStyle.cpp b/Engine/Source/Editor/EditorStyle/Private/StarshipStyle.cpp
--- a/Engine/Source/Editor/EditorStyle/Private/StarshipStyle.cpp
+++ b/Engine/Source/Editor/EditorStyle/Private/StarshipStyle.cpp
@@ -7531,6 +7531,7 @@ void FStarshipEditorStyle::FStyle::SetupToolkitStyles()
 		Set("CurveEditorTools.ActivateTransformTool", new IMAGE_BRUSH_SVG("Starship/GenericCurveEditor/ToolsTransform", Icon20x20));
 		Set("CurveEditorTools.ActivateRetimeTool", new IMAGE_BRUSH_SVG("Starship/GenericCurveEditor/ToolsRetime", Icon20x20));
 		Set("CurveEditorTools.ActivateMultiScaleTool", new IMAGE_BRUSH_SVG("Starship/GenericCurveEditor/ToolsMultiScale", Icon20x20));
+		Set("CurveEditorTools.ActivateLatticeTool", new IMAGE_BRUSH_SVG("Starship/GenericCurveEditor/ToolsMultiScale", Icon20x20));
 	}
 
 	// General Curve Icons
