Fix CIS on Clang

diff --git a/Engine/Plugins/Editor/CurveEditorTools/Source/CurveEditorTools/Private/Tools/Lattice/CurveEditorLatticeTool.cpp b/Engine/Plugins/Editor/CurveEditorTools/Source/CurveEditorTools/Private/Tools/Lattice/CurveEditorLatticeTool.cpp
--- a/Engine/Plugins/Editor/CurveEditorTools/Source/CurveEditorTools/Private/Tools/Lattice/CurveEditorLatticeTool.cpp
+++ b/Engine/Plugins/Editor/CurveEditorTools/Source/CurveEditorTools/Private/Tools/Lattice/CurveEditorLatticeTool.cpp
@@ -25,9 +25,9 @@ FCurveEditorLatticeTool::FCurveEditorLatticeTool(TWeakPtr<FCurveEditor> InCurveE
 		NewObject<UCurveEditorTools_LatticeUndoData>(GetTransientPackage(), UCurveEditorTools_LatticeUndoData::StaticClass(), NAME_None, RF_Transactional)
 	)
 {
-	UndoObject->OnPostEditUndo.AddLambda([this, WeakCurveEditor = WeakCurveEditor]
+	UndoObject->OnPostEditUndo.AddLambda([this, CurveEditor = InCurveEditor]
 	{
-		const TSharedPtr<FCurveEditor> Editor = WeakCurveEditor.Pin();
+		const TSharedPtr<FCurveEditor> Editor = CurveEditor.Pin();
 		if (!Editor)
 		{
 			return;
