Curve Editor: Add tweening tool to toolbar

diff --git a/Engine/Plugins/Editor/CurveEditorTools/CurveEditorTools.uplugin b/Engine/Plugins/Editor/CurveEditorTools/CurveEditorTools.uplugin
--- a/Engine/Plugins/Editor/CurveEditorTools/CurveEditorTools.uplugin
+++ b/Engine/Plugins/Editor/CurveEditorTools/CurveEditorTools.uplugin
@@ -14,6 +14,12 @@
 	"IsBetaVersion": false,
 	"Installed": false,
 	"EnabledByDefault": true,
+	"Plugins": [
+		{
+			"Name": "TweeningUtils",
+			"Enabled": true
+		}
+	],
 	"Modules": [
 		{
 			"Name": "CurveEditorTools",
diff --git a/Engine/Plugins/Editor/CurveEditorTools/Source/CurveEditorTools/CurveEditorTools.Build.cs b/Engine/Plugins/Editor/CurveEditorTools/Source/CurveEditorTools/CurveEditorTools.Build.cs
--- a/Engine/Plugins/Editor/CurveEditorTools/Source/CurveEditorTools/CurveEditorTools.Build.cs
+++ b/Engine/Plugins/Editor/CurveEditorTools/Source/CurveEditorTools/CurveEditorTools.Build.cs
@@ -26,7 +26,8 @@ public class CurveEditorTools : ModuleRules
 				"CoreUObject",
 				"AudioMixer",
 				"SignalProcessing",
-				"EditorStyle",
+				"EditorStyle", 
+				"TweeningUtilsEditor",
 			}
 		);
 	}
diff --git a/Engine/Plugins/Editor/CurveEditorTools/Source/CurveEditorTools/Private/CurveEditorFocusExtension.cpp b/Engine/Plugins/Editor/CurveEditorTools/Source/CurveEditorTools/Private/CurveEditorFocusExtension.cpp
deleted file mode 100644
diff --git a/Engine/Plugins/Editor/CurveEditorTools/Source/CurveEditorTools/Private/CurveEditorFocusExtension.h b/Engine/Plugins/Editor/CurveEditorTools/Source/CurveEditorTools/Private/CurveEditorFocusExtension.h
deleted file mode 100644
diff --git a/Engine/Plugins/Editor/CurveEditorTools/Source/CurveEditorTools/Private/CurveEditorToolsModule.cpp b/Engine/Plugins/Editor/CurveEditorTools/Source/CurveEditorTools/Private/CurveEditorToolsModule.cpp
--- a/Engine/Plugins/Editor/CurveEditorTools/Source/CurveEditorTools/Private/CurveEditorToolsModule.cpp
+++ b/Engine/Plugins/Editor/CurveEditorTools/Source/CurveEditorTools/Private/CurveEditorToolsModule.cpp
@@ -3,7 +3,8 @@
 #include "Modules/ModuleManager.h"
 #include "Framework/Commands/UICommandList.h"
 #include "ICurveEditorModule.h"
-#include "CurveEditorFocusExtension.h"
+#include "Extensions/CurveEditorFocusExtension.h"
+#include "Extensions/TweenEditorExtension.h"
 #include "CurveEditorToolCommands.h"
 #include "Tools/CurveEditorTransformTool.h"
 #include "Tools/CurveEditorRetimeTool.h"
@@ -22,6 +23,7 @@ class FCurveEditorToolsModule : public IModuleInterface, public TSharedFromThis<
 
 private:
 	static TSharedRef<ICurveEditorExtension> CreateFocusExtension(TWeakPtr<FCurveEditor> InCurveEditor);
+	static TSharedRef<ICurveEditorExtension> CreateTweenExtension(TWeakPtr<FCurveEditor> InCurveEditor);
 	static TUniquePtr<ICurveEditorToolExtension> CreateTransformToolExtension(TWeakPtr<FCurveEditor> InCurveEditor);
 	static TUniquePtr<ICurveEditorToolExtension> CreateRetimeToolExtension(TWeakPtr<FCurveEditor> InCurveEditor);
 	static TUniquePtr<ICurveEditorToolExtension> CreateMultiScaleToolExtension(TWeakPtr<FCurveEditor> InCurveEditor);
@@ -33,6 +35,8 @@ protected:
 
 private:
 	FDelegateHandle FocusExtensionsHandle;
+	FDelegateHandle TweenExtensionHandle;
+	
 	FDelegateHandle TransformToolHandle;
 	FDelegateHandle RetimeToolHandle;
 	FDelegateHandle MultiScaleToolHandle;
@@ -49,13 +53,13 @@ void FCurveEditorToolsModule::StartupModule()
 
 	// Register Editor Extensions
 	FocusExtensionsHandle = CurveEditorModule.RegisterEditorExtension(FOnCreateCurveEditorExtension::CreateStatic(&FCurveEditorToolsModule::CreateFocusExtension));
+	TweenExtensionHandle = CurveEditorModule.RegisterEditorExtension(FOnCreateCurveEditorExtension::CreateStatic(&FCurveEditorToolsModule::CreateTweenExtension));
 
 	// Register Tool Extensions
 	TransformToolHandle = CurveEditorModule.RegisterToolExtension(FOnCreateCurveEditorToolExtension::CreateStatic(&FCurveEditorToolsModule::CreateTransformToolExtension));
 	RetimeToolHandle = CurveEditorModule.RegisterToolExtension(FOnCreateCurveEditorToolExtension::CreateStatic(&FCurveEditorToolsModule::CreateRetimeToolExtension));
 	MultiScaleToolHandle = CurveEditorModule.RegisterToolExtension(FOnCreateCurveEditorToolExtension::CreateStatic(&FCurveEditorToolsModule::CreateMultiScaleToolExtension));
 	LatticeToolHandle = CurveEditorModule.RegisterToolExtension(FOnCreateCurveEditorToolExtension::CreateStatic(&FCurveEditorToolsModule::CreateLatticeToolExtensions));
-	
 
 	auto ToolbarExtender = ICurveEditorModule::FCurveEditorMenuExtender::CreateRaw(this, &FCurveEditorToolsModule::ExtendCurveEditorToolbarMenu);
 	auto& MenuExtenders = CurveEditorModule.GetAllToolBarMenuExtenders();
@@ -100,6 +104,7 @@ void FCurveEditorToolsModule::ShutdownModule()
 
 	// Unregister Editor Extensions
 	CurveEditorModule.UnregisterEditorExtension(FocusExtensionsHandle);
+	CurveEditorModule.UnregisterEditorExtension(TweenExtensionHandle);
 
 	// Unregister Tool Extensions
 	CurveEditorModule.UnregisterToolExtension(TransformToolHandle);
@@ -115,6 +120,11 @@ TSharedRef<ICurveEditorExtension> FCurveEditorToolsModule::CreateFocusExtension(
 	return MakeShared<FCurveEditorFocusExtension>(InCurveEditor);
 }
 
+TSharedRef<ICurveEditorExtension> FCurveEditorToolsModule::CreateTweenExtension(TWeakPtr<FCurveEditor> InCurveEditor)
+{
+	return MakeShared<UE::CurveEditor::FTweenEditorExtension>(InCurveEditor);
+}
+
 TUniquePtr<ICurveEditorToolExtension> FCurveEditorToolsModule::CreateTransformToolExtension(TWeakPtr<FCurveEditor> InCurveEditor)
 {
 	return MakeUnique<FCurveEditorTransformTool>(InCurveEditor);
diff --git a/Engine/Plugins/Editor/CurveEditorTools/Source/CurveEditorTools/Private/Extensions/CurveEditorFocusExtension.cpp b/Engine/Plugins/Editor/CurveEditorTools/Source/CurveEditorTools/Private/Extensions/CurveEditorFocusExtension.cpp
new file mode 100644
--- /dev/null
+++ b/Engine/Plugins/Editor/CurveEditorTools/Source/CurveEditorTools/Private/Extensions/CurveEditorFocusExtension.cpp
@@ -0,0 +1,109 @@
+// Copyright Epic Games, Inc. All Rights Reserved.
+
+#include "CurveEditorFocusExtension.h"
+#include "Framework/Commands/UICommandList.h"
+#include "CurveEditorToolCommands.h"
+#include "CurveEditor.h"
+#include "ICurveEditorBounds.h"
+#include "AnimatedRange.h"
+
+void FCurveEditorFocusExtension::BindCommands(TSharedRef<FUICommandList> CommandBindings)
+{
+	// ToDo: Focus on Timeline is only valid if there is a externally provided time controller.
+	FCanExecuteAction CanFocus = FCanExecuteAction::CreateSP(this, &FCurveEditorFocusExtension::CanUseFocusExtensions);
+	FExecuteAction FocusPlaybackTime = FExecuteAction::CreateSP(this, &FCurveEditorFocusExtension::FocusPlaybackTime);
+	FExecuteAction FocusPlaybackRange = FExecuteAction::CreateSP(this, &FCurveEditorFocusExtension::FocusPlaybackRange);
+
+	CommandBindings->MapAction(FCurveEditorToolCommands::Get().SetFocusPlaybackTime, FocusPlaybackTime, CanFocus);
+	CommandBindings->MapAction(FCurveEditorToolCommands::Get().SetFocusPlaybackRange, FocusPlaybackRange, CanFocus);
+}
+
+void FCurveEditorFocusExtension::FocusPlaybackRange()
+{
+	// This changes the extents of your view range to focus on the current playback range of the curve editor.
+	// This will change both zoom and position.
+	TSharedPtr<FCurveEditor> CurveEditor = WeakCurveEditor.Pin();
+	if (CurveEditor)
+	{
+		TSharedPtr<ITimeSliderController> TimeSliderController = CurveEditor->GetTimeSliderController();
+		if (TimeSliderController)
+		{
+			TRange<FFrameNumber> PlaybackRange = TimeSliderController->GetPlayRange();
+
+			// Tell the curve editor to Zoom to Fit on the Y axis. This will handle zooming in on the correct curves
+			// before we adjust the represented horizontal range.
+			CurveEditor->ZoomToFit(EAxisList::Type::Y);
+
+			// Now we need to convert between Time Slider Controller's playback range and the bounds.
+			FFrameNumber InFrame = TimeSliderController->GetPlayRange().GetLowerBoundValue();
+			FFrameNumber OutFrame = TimeSliderController->GetPlayRange().GetUpperBoundValue();
+
+			double InTime = InFrame / TimeSliderController->GetTickResolution();
+			double OutTime = OutFrame / TimeSliderController->GetTickResolution();
+
+			// Add 5% padding to each side
+			double Padding = (OutTime - InTime) * 0.05;
+			InTime -= Padding;
+			OutTime += Padding;
+
+			CurveEditor->GetBounds().SetInputBounds(InTime, OutTime);
+		}
+	}
+}
+
+void FCurveEditorFocusExtension::FocusPlaybackTime()
+{
+	// This takes your current view range (both x and y) and shifts it to be centered on the playback head.
+	// This does not change the extents of your view range so the zoom level will not shift.
+	TSharedPtr<FCurveEditor> CurveEditor = WeakCurveEditor.Pin();
+	if (CurveEditor)
+	{
+		TSharedPtr<ITimeSliderController> TimeSliderController = CurveEditor->GetTimeSliderController();
+		if (TimeSliderController)
+		{
+			double InTime, OutTime;
+			CurveEditor->GetBounds().GetInputBounds(InTime, OutTime);
+
+			// Calculate our total range
+			double ViewRange = OutTime - InTime;
+
+			const FFrameTime ScrubFrameTime = TimeSliderController->GetScrubPosition();
+
+			// We cant move outside clamp range, so take that into account
+			FAnimatedRange ClampRange = TimeSliderController->GetClampRange();
+
+			// Now center it around the Time Slider.
+			double PlaybackPosition = ScrubFrameTime / TimeSliderController->GetTickResolution();
+			InTime = PlaybackPosition - (ViewRange / 2.0);
+
+			// Take clamp range into account
+			if(InTime < ClampRange.GetLowerBoundValue())
+			{
+				InTime = ClampRange.GetLowerBoundValue();
+				OutTime = InTime + ViewRange;
+			}
+			else
+			{
+				OutTime = PlaybackPosition + (ViewRange / 2.0);
+				if(OutTime > ClampRange.GetUpperBoundValue())
+				{
+					OutTime = ClampRange.GetUpperBoundValue();
+					InTime = OutTime - ViewRange;
+				}
+			}
+
+			CurveEditor->GetBounds().SetInputBounds(InTime, OutTime);
+		}
+	}
+}
+
+bool FCurveEditorFocusExtension::CanUseFocusExtensions() const
+{
+	TSharedPtr<FCurveEditor> CurveEditor = WeakCurveEditor.Pin();
+	if (CurveEditor)
+	{
+		return CurveEditor->GetTimeSliderController().IsValid();
+	}
+
+	return false;
+}
diff --git a/Engine/Plugins/Editor/CurveEditorTools/Source/CurveEditorTools/Private/Extensions/CurveEditorFocusExtension.h b/Engine/Plugins/Editor/CurveEditorTools/Source/CurveEditorTools/Private/Extensions/CurveEditorFocusExtension.h
new file mode 100644
--- /dev/null
+++ b/Engine/Plugins/Editor/CurveEditorTools/Source/CurveEditorTools/Private/Extensions/CurveEditorFocusExtension.h
@@ -0,0 +1,31 @@
+// Copyright Epic Games, Inc. All Rights Reserved.
+
+#pragma once
+
+#include "ICurveEditorExtension.h"
+
+class FCurveEditor;
+
+class FCurveEditorFocusExtension : public ICurveEditorExtension, public TSharedFromThis<FCurveEditorFocusExtension>
+{
+public:
+	FCurveEditorFocusExtension(TWeakPtr<FCurveEditor> InCurveEditor)
+		: WeakCurveEditor(InCurveEditor)
+	{
+	}
+
+	// ICurveEditorExtension Interface
+	virtual void BindCommands(TSharedRef<FUICommandList> CommandBindings) override;
+	// ~ICurveEditorExtension
+
+private:
+	/** Focuses the curve editor around the current Playback Range of the linked Time Controller. */
+	void FocusPlaybackRange();
+	/** Focuses the curve editor around the current Playback Time of the linked Time Controller without changing zoom level. */
+	void FocusPlaybackTime();
+
+	/** Checks to see if the Curve Editor has a valid Time Controller to use these focus extensions on. */
+	bool CanUseFocusExtensions() const;
+private:
+	TWeakPtr<FCurveEditor> WeakCurveEditor;
+};
\ No newline at end of file
diff --git a/Engine/Plugins/Editor/CurveEditorTools/Source/CurveEditorTools/Private/Extensions/TweenEditorExtension.cpp b/Engine/Plugins/Editor/CurveEditorTools/Source/CurveEditorTools/Private/Extensions/TweenEditorExtension.cpp
new file mode 100644
--- /dev/null
+++ b/Engine/Plugins/Editor/CurveEditorTools/Source/CurveEditorTools/Private/Extensions/TweenEditorExtension.cpp
@@ -0,0 +1,82 @@
+﻿// Copyright Epic Games, Inc. All Rights Reserved.
+
+#include "TweenEditorExtension.h"
+
+#include "Framework/MultiBox/MultiBoxBuilder.h"
+#include "Framework/MultiBox/MultiBoxExtender.h"
+#include "ICurveEditorModule.h"
+#include "Math/Abstraction/BuiltInTweenFunctionContainer.h"
+#include "Widgets/MVC/Models/CurveTweenModel.h"
+#include "Widgets/MVC/Models/TransactionalTweenModelProxy.h"
+#include "Widgets/MVC/TweenToolbarController.h"
+
+namespace UE::CurveEditor
+{
+FTweenEditorExtension::FTweenEditorExtension(TWeakPtr<FCurveEditor> InCurveEditor)
+	: WeakCurveEditor(InCurveEditor)
+	, BlendFunctionContainer(MakeShared<TweeningUtilsEditor::FBuiltInTweenFunctionContainer>())
+{
+	ICurveEditorModule::FCurveEditorMenuExtender ToolbarExtender = ICurveEditorModule::FCurveEditorMenuExtender::CreateRaw(
+		this, &FTweenEditorExtension::ExtendCurveEditorToolbarMenu
+		);
+	TArray<ICurveEditorModule::FCurveEditorMenuExtender>& MenuExtenders = ICurveEditorModule::Get().GetAllToolBarMenuExtenders();
+	MenuExtenders.Add(ToolbarExtender);
+}
+
+FTweenEditorExtension::~FTweenEditorExtension()
+{
+	TArray<ICurveEditorModule::FCurveEditorMenuExtender>& MenuExtenders = ICurveEditorModule::Get().GetAllToolBarMenuExtenders();
+	for (auto It = MenuExtenders.CreateIterator(); It; ++It)
+	{
+		if (It->IsBoundToObject(this))
+		{
+			It.RemoveCurrent();
+		}
+	}
+}
+
+void FTweenEditorExtension::BindCommands(TSharedRef<FUICommandList> InCommandList)
+{
+	GetOrCreateToolbarController(InCommandList);
+}
+
+TSharedRef<FExtender> FTweenEditorExtension::ExtendCurveEditorToolbarMenu(const TSharedRef<FUICommandList> InCommandList)
+{
+	TSharedRef<FExtender> Extender = MakeShared<FExtender>();
+	
+	Extender->AddToolBarExtension("Tangents", EExtensionHook::After, InCommandList, FToolBarExtensionDelegate::CreateLambda(
+		[this, WeakCommandList = InCommandList.ToWeakPtr()](FToolBarBuilder& ToolbarBuilder)
+		{
+			const TSharedPtr<FUICommandList> CommandListPin = WeakCommandList.Pin();
+			if (ensure(CommandListPin) && GetOrCreateToolbarController(CommandListPin.ToSharedRef()))
+			{
+				ToolbarBuilder.BeginSection("Tween");
+				ToolbarController->AddToToolbar(ToolbarBuilder);
+				ToolbarBuilder.EndSection();
+			}
+		}));
+	
+	return Extender;
+}
+
+TSharedPtr<TweeningUtilsEditor::FTweenToolbarController> FTweenEditorExtension::GetOrCreateToolbarController(TSharedRef<FUICommandList> InCommandList)
+{
+	if (ToolbarController)
+	{
+		return ToolbarController;
+	}
+
+	if (const TSharedPtr<FCurveEditor> CurveEditorPin = WeakCurveEditor.Pin(); ensure(CurveEditorPin))
+	{
+		using namespace TweeningUtilsEditor;
+		
+		const TSharedRef<TTransactionalTweenModelProxy<FCurveTweenModel>> TweenModel = MakeShared<TTransactionalTweenModelProxy<FCurveTweenModel>>(
+			WeakCurveEditor,
+			BlendFunctionContainer->GetFunction(0).ToSharedRef()
+			);
+		ToolbarController = MakeShared<FTweenToolbarController>(TweenModel, InCommandList, BlendFunctionContainer);
+	}
+	
+	return ToolbarController;
+}
+}
diff --git a/Engine/Plugins/Editor/CurveEditorTools/Source/CurveEditorTools/Private/Extensions/TweenEditorExtension.h b/Engine/Plugins/Editor/CurveEditorTools/Source/CurveEditorTools/Private/Extensions/TweenEditorExtension.h
new file mode 100644
--- /dev/null
+++ b/Engine/Plugins/Editor/CurveEditorTools/Source/CurveEditorTools/Private/Extensions/TweenEditorExtension.h
@@ -0,0 +1,44 @@
+﻿// Copyright Epic Games, Inc. All Rights Reserved.
+
+#pragma once
+
+#include "ICurveEditorExtension.h"
+#include "Templates/SharedPointer.h"
+
+class FCurveEditor;
+class FExtender;
+namespace UE::TweeningUtilsEditor { class FBuiltInTweenFunctionContainer; }
+namespace UE::TweeningUtilsEditor { class FTweenToolbarController; }
+
+namespace UE::CurveEditor
+{
+/** Manages the tweening tool in the curve editor. */
+class FTweenEditorExtension : public ICurveEditorExtension, public FNoncopyable
+{
+public:
+
+	explicit FTweenEditorExtension(TWeakPtr<FCurveEditor> InCurveEditor);
+	virtual ~FTweenEditorExtension() override;
+
+	//~ Begin ICurveEditorExtension Interface
+	virtual void BindCommands(TSharedRef<FUICommandList> InCommandList) override;
+	//~ End ICurveEditorExtension Interface
+
+private:
+
+	/** The editor that owns us. */
+	const TWeakPtr<FCurveEditor> WeakCurveEditor;
+
+	/** Holds the functions that ToolbarController shows: just the default, built-in ones.*/
+	const TSharedRef<TweeningUtilsEditor::FBuiltInTweenFunctionContainer> BlendFunctionContainer;
+	/** Created by BindCommands. */
+	TSharedPtr<TweeningUtilsEditor::FTweenToolbarController> ToolbarController;
+
+	/** Adds the tweening widgets to the toolbar. */
+	TSharedRef<FExtender> ExtendCurveEditorToolbarMenu(const TSharedRef<FUICommandList> InCommandList);
+
+	/** Makes sure that ToolbarController is created. */
+	TSharedPtr<TweeningUtilsEditor::FTweenToolbarController> GetOrCreateToolbarController(TSharedRef<FUICommandList> InCommandList);
+};
+}
+
