Curve editor: Fix crash when creating a 2nd level sequence.
- Each FTweenEditorExtension registered a toolbar extender.
- Due to order of operations, "old" extensions are alive a little too long so 2 instances of FTweenEditorExtension extended the toolbar.
- The old instance was destroyed but the combo button it added to the toolbar was referencing CreateRaw delegates pointing to memory that belonged to the now destroyed, old FTweenEditorExtension
- Fix: instead of having each FTweenEditorExtension register global toolbar extender, allow ICurveEditorExtension to extend the toolbar directly. This makes other extensions safer down the road, too.

diff --git a/Engine/Plugins/Editor/CurveEditorTools/Source/CurveEditorTools/Private/Extensions/TweenEditorExtension.cpp b/Engine/Plugins/Editor/CurveEditorTools/Source/CurveEditorTools/Private/Extensions/TweenEditorExtension.cpp
--- a/Engine/Plugins/Editor/CurveEditorTools/Source/CurveEditorTools/Private/Extensions/TweenEditorExtension.cpp
+++ b/Engine/Plugins/Editor/CurveEditorTools/Source/CurveEditorTools/Private/Extensions/TweenEditorExtension.cpp
@@ -15,32 +15,14 @@ namespace UE::CurveEditor
 FTweenEditorExtension::FTweenEditorExtension(TWeakPtr<FCurveEditor> InCurveEditor)
 	: WeakCurveEditor(InCurveEditor)
 	, BlendFunctionContainer(MakeShared<TweeningUtilsEditor::FBuiltInTweenFunctionContainer>())
-{
-	ICurveEditorModule::FCurveEditorMenuExtender ToolbarExtender = ICurveEditorModule::FCurveEditorMenuExtender::CreateRaw(
-		this, &FTweenEditorExtension::ExtendCurveEditorToolbarMenu
-		);
-	TArray<ICurveEditorModule::FCurveEditorMenuExtender>& MenuExtenders = ICurveEditorModule::Get().GetAllToolBarMenuExtenders();
-	MenuExtenders.Add(ToolbarExtender);
-}
-
-FTweenEditorExtension::~FTweenEditorExtension()
-{
-	TArray<ICurveEditorModule::FCurveEditorMenuExtender>& MenuExtenders = ICurveEditorModule::Get().GetAllToolBarMenuExtenders();
-	for (auto It = MenuExtenders.CreateIterator(); It; ++It)
-	{
-		if (It->IsBoundToObject(this))
-		{
-			It.RemoveCurrent();
-		}
-	}
-}
+{}
 
 void FTweenEditorExtension::BindCommands(TSharedRef<FUICommandList> InCommandList)
 {
 	GetOrCreateToolbarController(InCommandList);
 }
 
-TSharedRef<FExtender> FTweenEditorExtension::ExtendCurveEditorToolbarMenu(const TSharedRef<FUICommandList> InCommandList)
+TSharedPtr<FExtender> FTweenEditorExtension::MakeToolbarExtender(const TSharedRef<FUICommandList>& InCommandList)
 {
 	TSharedRef<FExtender> Extender = MakeShared<FExtender>();
 	
diff --git a/Engine/Plugins/Editor/CurveEditorTools/Source/CurveEditorTools/Private/Extensions/TweenEditorExtension.h b/Engine/Plugins/Editor/CurveEditorTools/Source/CurveEditorTools/Private/Extensions/TweenEditorExtension.h
--- a/Engine/Plugins/Editor/CurveEditorTools/Source/CurveEditorTools/Private/Extensions/TweenEditorExtension.h
+++ b/Engine/Plugins/Editor/CurveEditorTools/Source/CurveEditorTools/Private/Extensions/TweenEditorExtension.h
@@ -18,10 +18,10 @@ class FTweenEditorExtension : public ICurveEditorExtension, public FNoncopyable
 public:
 
 	explicit FTweenEditorExtension(TWeakPtr<FCurveEditor> InCurveEditor);
-	virtual ~FTweenEditorExtension() override;
 
 	//~ Begin ICurveEditorExtension Interface
 	virtual void BindCommands(TSharedRef<FUICommandList> InCommandList) override;
+	virtual TSharedPtr<FExtender> MakeToolbarExtender(const TSharedRef<FUICommandList>& InCommandList) override;
 	//~ End ICurveEditorExtension Interface
 
 private:
@@ -34,9 +34,6 @@ private:
 	/** Created by BindCommands. */
 	TSharedPtr<TweeningUtilsEditor::FTweenToolbarController> ToolbarController;
 
-	/** Adds the tweening widgets to the toolbar. */
-	TSharedRef<FExtender> ExtendCurveEditorToolbarMenu(const TSharedRef<FUICommandList> InCommandList);
-
 	/** Makes sure that ToolbarController is created. */
 	TSharedPtr<TweeningUtilsEditor::FTweenToolbarController> GetOrCreateToolbarController(TSharedRef<FUICommandList> InCommandList);
 };
diff --git a/Engine/Source/Editor/CurveEditor/Private/SCurveEditorPanel.cpp b/Engine/Source/Editor/CurveEditor/Private/SCurveEditorPanel.cpp
--- a/Engine/Source/Editor/CurveEditor/Private/SCurveEditorPanel.cpp
+++ b/Engine/Source/Editor/CurveEditor/Private/SCurveEditorPanel.cpp
@@ -34,6 +34,7 @@
 #include "HAL/IConsoleManager.h"
 #include "HAL/PlatformCrt.h"
 #include "ICurveEditorDragOperation.h"
+#include "ICurveEditorExtension.h"
 #include "ICurveEditorModule.h"
 #include "ICurveEditorToolExtension.h"
 #include "IPropertyRowGenerator.h"
@@ -1197,9 +1198,17 @@ TSharedPtr<FExtender> SCurveEditorPanel::GetToolbarExtender()
 	// We combine this with the extender provided by the Curve Editor Module as that extender has been extended by tools
 	ICurveEditorModule& CurveEditorModule = FModuleManager::Get().LoadModuleChecked<ICurveEditorModule>("CurveEditor");
 	TArray<TSharedPtr<FExtender>> ToolbarExtenders;
+	const TSharedRef<FUICommandList> CommandListRef = GetCommands().ToSharedRef();
 	for (ICurveEditorModule::FCurveEditorMenuExtender& ExtenderCallback : CurveEditorModule.GetAllToolBarMenuExtenders())
 	{
-		ToolbarExtenders.Add(ExtenderCallback.Execute(GetCommands().ToSharedRef()));
+		ToolbarExtenders.Add(ExtenderCallback.Execute(CommandListRef));
+	}
+	for (const TSharedRef<ICurveEditorExtension>& Extension : CurveEditor->GetEditorExtensions())
+	{
+		if (const TSharedPtr<FExtender> OptionalExtender = Extension->MakeToolbarExtender(CommandListRef))
+		{
+			ToolbarExtenders.Add(OptionalExtender);
+		}
 	}
 	TSharedPtr<FExtender> Extender = FExtender::Combine(ToolbarExtenders);
 
diff --git a/Engine/Source/Editor/CurveEditor/Public/ICurveEditorExtension.h b/Engine/Source/Editor/CurveEditor/Public/ICurveEditorExtension.h
--- a/Engine/Source/Editor/CurveEditor/Public/ICurveEditorExtension.h
+++ b/Engine/Source/Editor/CurveEditor/Public/ICurveEditorExtension.h
@@ -17,4 +17,11 @@ public:
 	 * @param CommandBindings The existing command bindings to map to.
 	*/
 	virtual void BindCommands(TSharedRef<FUICommandList> CommandBindings) = 0;
+
+	/**
+	 * Allows the editors to extend the toolbar.
+	 *
+	 * @param InCommandList The existing command bindings to map to.
+	 */
+	virtual TSharedPtr<FExtender> MakeToolbarExtender(const TSharedRef<FUICommandList>& InCommandList) { return nullptr; }
 };
\ No newline at end of file
