Curve Editor: Tweening combo button now has a min size. When the user selects a tweening function, the combo button will not change its size (which would change the layout of the toolbar - quite distracting).

diff --git a/Engine/Plugins/Editor/CurveEditorTools/Source/CurveEditorTools/Private/Extensions/TweenEditorExtension.cpp b/Engine/Plugins/Editor/CurveEditorTools/Source/CurveEditorTools/Private/Extensions/TweenEditorExtension.cpp
--- a/Engine/Plugins/Editor/CurveEditorTools/Source/CurveEditorTools/Private/Extensions/TweenEditorExtension.cpp
+++ b/Engine/Plugins/Editor/CurveEditorTools/Source/CurveEditorTools/Private/Extensions/TweenEditorExtension.cpp
@@ -40,9 +40,11 @@ TSharedPtr<FExtender> FTweenEditorExtension::MakeToolbarExtender(const TSharedRe
 						return TweenControllers->MouseSlidingController.GetCurrentSliderPosition();
 					});
 				
+				ToolbarBuilder.BeginStyleOverride(TEXT("CurveEditorTweenToolbar")); // Forces min sizes for combo button
 				ToolbarBuilder.BeginSection("Tween");
 				TweenControllers->ToolbarController.AddToToolbar(ToolbarBuilder, MoveTemp(OverrideSliderPosition));
 				ToolbarBuilder.EndSection();
+				ToolbarBuilder.EndStyleOverride();
 			}
 		}));
 	
diff --git a/Engine/Source/Editor/EditorStyle/Private/StarshipStyle.cpp b/Engine/Source/Editor/EditorStyle/Private/StarshipStyle.cpp
--- a/Engine/Source/Editor/EditorStyle/Private/StarshipStyle.cpp
+++ b/Engine/Source/Editor/EditorStyle/Private/StarshipStyle.cpp
@@ -7457,6 +7457,10 @@ void FStarshipEditorStyle::FStyle::SetupToolkitStyles()
 		CurveEditorToolbar.SetSeparatorPadding(    FMargin(2.0f, 0.0f));
 
 		Set("CurveEditorToolBar", CurveEditorToolbar);
+		
+		CurveEditorToolbar.SetComboLabelMinWidth(95.f);
+		CurveEditorToolbar.SetComboContentHorizontalAlignment(HAlign_Center);
+		Set("CurveEditorTweenToolbar", CurveEditorToolbar);
 
 		// Tab
 		Set("GenericCurveEditor", new IMAGE_BRUSH_SVG("Starship/GenericCurveEditor/CurveEditor", Icon20x20));
