Curve Editor: Tools now flatten the tangents when keys are moved:
- Lattice tool:
- You can double click the lattice to quickly mirror all keys and tangents around each curve's mid-point
- Dragging the top or bottom edge over the opposite edge slowly interpolates tangents to 0 and then flips them.
- Tangents slowly approach 0 as the difference between min and max key height approaches 0, afterwards they flip:
- Tweening tool
- Multi Scale tool
- Transform tool

diff --git a/Engine/Plugins/Editor/CurveEditorTools/Source/CurveEditorTools/Private/Extensions/Tweening/TangentFlatteningTweenProxy.h b/Engine/Plugins/Editor/CurveEditorTools/Source/CurveEditorTools/Private/Extensions/Tweening/TangentFlatteningTweenProxy.h
new file mode 100644
--- /dev/null
+++ b/Engine/Plugins/Editor/CurveEditorTools/Source/CurveEditorTools/Private/Extensions/Tweening/TangentFlatteningTweenProxy.h
@@ -0,0 +1,55 @@
+﻿// Copyright Epic Games, Inc. All Rights Reserved.
+
+#pragma once
+
+#include "CurveEditor.h"
+#include "Misc/Mirror/TangentSelectionFlattener.h"
+#include "Widgets/MVC/Models/TweenModel.h"
+#include <type_traits>
+
+namespace UE::TweeningUtilsEditor
+{
+/**
+ * This class squishes the curves based on how much the tween function squishes the keys vertically.
+ * The squishing is achieved by interpolating the tangents to 0.
+ */
+template<typename TBase> requires std::is_base_of_v<FTweenModel, TBase>
+class TTangentFlatteningTweenProxy : public TBase
+{
+public:
+
+	template<typename... TArg>
+	explicit TTangentFlatteningTweenProxy(TWeakPtr<FCurveEditor> WeakCurveEditor, TArg&&... Arg)
+		: TBase(Forward<TArg>(Arg)...)
+		, WeakCurveEditor(MoveTemp(WeakCurveEditor))
+	{}
+
+	//~ Begin FTweenModel Interface
+	virtual void StartBlendOperation() override
+	{
+		TBase::StartBlendOperation();
+		
+		if (const TSharedPtr<FCurveEditor> CurveEditorPin = WeakCurveEditor.Pin())
+		{
+			TangentTweener.ResetFromSelection(*CurveEditorPin);
+		}
+	}
+	virtual void BlendValues(float InNormalizedValue) override
+	{
+		TBase::BlendValues(InNormalizedValue);
+		if (const TSharedPtr<FCurveEditor> CurveEditorPin = WeakCurveEditor.Pin())
+		{
+			TangentTweener.ComputeMirroringParallel(*CurveEditorPin);
+		}
+	}
+	//~ Begin FTweenModel Interface
+
+private:
+
+	/** Needed as arg for TangentTweener. */
+	const TWeakPtr<FCurveEditor> WeakCurveEditor;
+	/** Implements the logic for flattening the tangents. */
+	CurveEditor::FTangentSelectionFlattener TangentTweener;
+};
+}
+
diff --git a/Engine/Plugins/Editor/CurveEditorTools/Source/CurveEditorTools/Private/Extensions/Tweening/TweenEditorExtension.cpp b/Engine/Plugins/Editor/CurveEditorTools/Source/CurveEditorTools/Private/Extensions/Tweening/TweenEditorExtension.cpp
--- a/Engine/Plugins/Editor/CurveEditorTools/Source/CurveEditorTools/Private/Extensions/Tweening/TweenEditorExtension.cpp
+++ b/Engine/Plugins/Editor/CurveEditorTools/Source/CurveEditorTools/Private/Extensions/Tweening/TweenEditorExtension.cpp
@@ -9,6 +9,7 @@
 #include "Framework/MultiBox/MultiBoxExtender.h"
 #include "Math/Abstraction/BuiltInTweenFunctionContainer.h"
 #include "Misc/Optional.h"
+#include "TangentFlatteningTweenProxy.h"
 #include "Widgets/MVC/Models/TransactionalTweenModelProxy.h"
 #include "Widgets/MVC/TweenToolbarController.h"
 
@@ -73,7 +74,6 @@ bool FTweenEditorExtension::InitControllers(TSharedRef<FUICommandList> InCommand
 	if (const TSharedPtr<FCurveEditor> CurveEditorPin = WeakCurveEditor.Pin(); ensure(CurveEditorPin))
 	{
 		using namespace TweeningUtilsEditor;
-		using FModel = FContextAwareTweeningModel;
 
 		const TSharedPtr<ITweenFunction> PreferredFunction = BlendFunctionContainer->FindFunctionByIdentifier(
 			UCurveEditorToolsUserSettings::Get()->GetPreferredTweenFunction()
@@ -81,8 +81,10 @@ bool FTweenEditorExtension::InitControllers(TSharedRef<FUICommandList> InCommand
 		const TSharedRef<ITweenFunction> InitialFunction = PreferredFunction
 			? PreferredFunction.ToSharedRef() : BlendFunctionContainer->GetFunction(0).ToSharedRef();
 		
-		const TSharedRef<TTransactionalTweenModelProxy<FModel>> TweenModel = MakeShared<TTransactionalTweenModelProxy<FModel>>(
-			WeakCurveEditor, InitialFunction
+		using FModel = TTransactionalTweenModelProxy<TTangentFlatteningTweenProxy<FContextAwareTweeningModel>>;
+		const TSharedRef<FModel> TweenModel = MakeShared<FModel>(
+			WeakCurveEditor,					// TTangentFlatteningTweenProxy arg
+			WeakCurveEditor, InitialFunction	// FContextAwareTweeningModel args
 			);
 		TweenControllers.Emplace(TweenModel, InCommandList, BlendFunctionContainer);
 		TweenControllers->ToolbarController.OnTweenFunctionChanged().AddRaw(this, &FTweenEditorExtension::OnTweenFunctionSelected);
diff --git a/Engine/Plugins/Editor/CurveEditorTools/Source/CurveEditorTools/Private/Misc/KeyInCurveId.cpp b/Engine/Plugins/Editor/CurveEditorTools/Source/CurveEditorTools/Private/Misc/KeyInCurveId.cpp
new file mode 100644
--- /dev/null
+++ b/Engine/Plugins/Editor/CurveEditorTools/Source/CurveEditorTools/Private/Misc/KeyInCurveId.cpp
@@ -0,0 +1,36 @@
+﻿// Copyright Epic Games, Inc. All Rights Reserved.
+
+#include "Misc/KeyInCurveId.h"
+
+#include "CurveEditor.h"
+#include "Containers/Map.h"
+
+namespace UE::CurveEditorTools
+{
+void ModifyCurves(TConstArrayView<FKeyInCurveId> InKeysToMirror, const FCurveEditor& InCurveEditor)
+{
+	// There's no harm calling Modify multiple times, but it's a waste of operations.
+	// Array instead of set because there will be a relatively low number of curves (<1000).
+	TArray<FCurveModelID> ModifiedCurves; 
+	for (const FKeyInCurveId& KeyInCurveId : InKeysToMirror)
+	{
+		const FCurveModelID CurveId = KeyInCurveId.CurveId;
+		FCurveModel* Model = InCurveEditor.FindCurve(CurveId);
+		if (Model && !ModifiedCurves.Contains(CurveId))
+		{
+			ModifiedCurves.AddUnique(CurveId);
+			Model->Modify();
+		}
+	}
+}
+	
+TMap<FCurveModelID, TArray<FKeyHandle>> SplitKeysInCurves(TConstArrayView<FKeyInCurveId> InKeys)
+{
+	TMap<FCurveModelID, TArray<FKeyHandle>> Result;
+	for (const FKeyInCurveId& KeyInCurveId : InKeys)
+	{
+		Result.FindOrAdd(KeyInCurveId.CurveId).Add(KeyInCurveId.KeyId);
+	}
+	return Result;
+}
+}
diff --git a/Engine/Plugins/Editor/CurveEditorTools/Source/CurveEditorTools/Private/Misc/KeyInCurveId.h b/Engine/Plugins/Editor/CurveEditorTools/Source/CurveEditorTools/Private/Misc/KeyInCurveId.h
new file mode 100644
--- /dev/null
+++ b/Engine/Plugins/Editor/CurveEditorTools/Source/CurveEditorTools/Private/Misc/KeyInCurveId.h
@@ -0,0 +1,33 @@
+﻿// Copyright Epic Games, Inc. All Rights Reserved.
+
+#pragma once
+
+#include "Containers/ContainersFwd.h"
+#include "Curves/KeyHandle.h"
+#include "CurveEditorTypes.h"
+#include "Serialization/Archive.h"
+
+class FCurveEditor;
+
+namespace UE::CurveEditorTools
+{
+/** Identifies a key on a curve. */
+struct FKeyInCurveId
+{
+	FKeyHandle KeyId;
+	FCurveModelID CurveId;
+
+	friend FArchive& operator<<(FArchive& InArchive, FKeyInCurveId& InCurveData)
+	{
+		InArchive << InCurveData.KeyId;
+		InArchive << InCurveData.CurveId;
+		return InArchive;
+	}
+};
+
+/** Calls Modify on the referenced curves. */
+void ModifyCurves(TConstArrayView<FKeyInCurveId> InKeysToMirror, const FCurveEditor& InCurveEditor);
+	
+/** Flattens the array a mapping of curve ids to the handles that belong to themm */
+TMap<FCurveModelID, TArray<FKeyHandle>> SplitKeysInCurves(TConstArrayView<FKeyInCurveId> InKeys);
+}
diff --git a/Engine/Plugins/Editor/CurveEditorTools/Source/CurveEditorTools/Private/Tools/CurveEditorMultiScaleTool.cpp b/Engine/Plugins/Editor/CurveEditorTools/Source/CurveEditorTools/Private/Tools/CurveEditorMultiScaleTool.cpp
--- a/Engine/Plugins/Editor/CurveEditorTools/Source/CurveEditorTools/Private/Tools/CurveEditorMultiScaleTool.cpp
+++ b/Engine/Plugins/Editor/CurveEditorTools/Source/CurveEditorTools/Private/Tools/CurveEditorMultiScaleTool.cpp
@@ -9,6 +9,7 @@
 #include "SCurveEditorView.h"
 #include "SCurveEditorPanel.h"
 #include "Algo/MaxElement.h"
+#include "Misc/Mirror/TangentSelectionFlattener.h"
 
 #include UE_INLINE_GENERATED_CPP_BY_NAME(CurveEditorMultiScaleTool)
 
@@ -357,6 +358,7 @@ void FCurveEditorMultiScaleTool::OnToolOptionsUpdated(const FPropertyChangedEven
 	ToolOptions.XScale = 1.f;
 	ToolOptions.YScale = 1.f;
 	ActiveTransaction.Reset();
+	TangentFlattener.Reset();
 }
 
 FText FCurveEditorMultiScaleTool::GetLabel() const
@@ -383,6 +385,8 @@ void FCurveEditorMultiScaleTool::OnDragStart()
 	{
 		return;
 	}
+	
+	TangentFlattener.Emplace(*CurveEditor);
 
 	// We need to cache our key data because all of our calculations have to be relative to the starting data and not the current per-frame data.
 	KeysByCurve.Reset();
@@ -581,24 +585,21 @@ void FCurveEditorMultiScaleTool::UpdateBoundingBox()
 }
 
 
-void FCurveEditorMultiScaleTool::ScaleUnique(const FVector2D& InChangeAmount, const bool bInAffectsX, const bool bInAffectsY) const
+void FCurveEditorMultiScaleTool::ScaleUnique(const FVector2D& InChangeAmount, const bool bInAffectsX, const bool bInAffectsY)
 {
 	TSharedPtr<FCurveEditor> CurveEditor = WeakCurveEditor.Pin();
 	if (!CurveEditor.IsValid())
 	{
 		return;
 	}
-
+	
 	// Rescaling is where things get tricky. If a user is selecting an edge we scale on one axis at once, and if they select a corner we scale on two axis at once.
 	// If they press alt, we scale relative to the center (instead of relative to the opposite corner). Because these are all bit-flag'd together we can implement
 	// this generically by scaling coordinates relative to an arbitrary center
-
 	const FVector2D PercentChanged = FVector2D(1.0f, 1.0f) + InChangeAmount; // ie: 5 pixel change on a 100 wide gives you 1.05
-
-	FSlateLayoutTransform ContainerToAbsolute = CurveEditor->GetPanel()->GetViewContainerGeometry().GetAccumulatedLayoutTransform().Inverse();
-
+	
 	TArray<FKeyPosition> NewKeyPositionScratch;
-
+	
 	// We now know if we need to affect both X and Y, and we know where we're scaling from. Now we can loop through the keys and actually modify their positions.
 	// We perform the scale on both axis (for simplicity) and then read which axis it should effect before assigning it back to the key position.
 	for (const FKeyData& KeyData : KeysByCurve)
@@ -646,6 +647,13 @@ void FCurveEditorMultiScaleTool::ScaleUnique(const FVector2D& InChangeAmount, co
 
 		CurveModel->SetKeyPositions(KeyData.Handles, NewKeyPositionScratch);
 	}
+	
+	// This recalculates the user specified tangents based on the distance moved. We'll flatten the curves depending on the resulting selection's height.
+	if (bInAffectsY)
+	{
+		const bool bIsMirroring = PercentChanged.Y < 0.0;
+		TangentFlattener->ComputeMirroringParallel(*CurveEditor, bIsMirroring);
+	}
 }
 
 FVector2D FCurveEditorMultiScaleTool::GetPivot(FCurveModel* InCurve, const TArray<FKeyPosition>& InKeyPositions) const
diff --git a/Engine/Plugins/Editor/CurveEditorTools/Source/CurveEditorTools/Private/Tools/CurveEditorMultiScaleTool.h b/Engine/Plugins/Editor/CurveEditorTools/Source/CurveEditorTools/Private/Tools/CurveEditorMultiScaleTool.h
--- a/Engine/Plugins/Editor/CurveEditorTools/Source/CurveEditorTools/Private/Tools/CurveEditorMultiScaleTool.h
+++ b/Engine/Plugins/Editor/CurveEditorTools/Source/CurveEditorTools/Private/Tools/CurveEditorMultiScaleTool.h
@@ -1,13 +1,14 @@
 // Copyright Epic Games, Inc. All Rights Reserved.
 
 #pragma once
+
 #include "ICurveEditorToolExtension.h"
 #include "Framework/DelayedDrag.h"
 #include "CurveDataAbstraction.h"
 #include "CurveEditorTypes.h"
 #include "ScopedTransaction.h"
+#include "Misc/Mirror/TangentSelectionFlattener.h"
 #include "UObject/StructOnScope.h"
-
 #include "CurveEditorMultiScaleTool.generated.h"
 
 class FCurveEditor;
@@ -125,7 +126,7 @@ private:
 	void StopDragIfPossible();
 	void UpdateBoundingBox();
 
-	void ScaleUnique(const FVector2D& InChangeAmount, const bool bInAffectsX, const bool bInAffectsY) const;
+	void ScaleUnique(const FVector2D& InChangeAmount, const bool bInAffectsX, const bool bInAffectsY);
 
 private:
 	/** Weak pointer back to the owning curve editor. */
@@ -133,10 +134,12 @@ private:
 
 	/** The currently open transaction (if any) */
 	TUniquePtr<class FScopedTransaction> ActiveTransaction;
-
+	
 	/** Set when attempting to move a drag handle. */
 	TOptional<FDelayedDrag> DelayedDrag;
 	TOptional<FVector2D> DragDelta;
+	/** Squishes the tangents vertically depending on how much we move the keys vertically. */
+	TOptional<UE::CurveEditor::FTangentSelectionFlattener> TangentFlattener;
 
 	/** Used to cache selected key data when doing transform operations */
 	struct FKeyData
diff --git a/Engine/Plugins/Editor/CurveEditorTools/Source/CurveEditorTools/Private/Tools/CurveEditorTransformTool.cpp b/Engine/Plugins/Editor/CurveEditorTools/Source/CurveEditorTools/Private/Tools/CurveEditorTransformTool.cpp
--- a/Engine/Plugins/Editor/CurveEditorTools/Source/CurveEditorTools/Private/Tools/CurveEditorTransformTool.cpp
+++ b/Engine/Plugins/Editor/CurveEditorTools/Source/CurveEditorTools/Private/Tools/CurveEditorTransformTool.cpp
@@ -511,6 +511,7 @@ void FCurveEditorTransformTool::OnToolOptionsUpdated(const FPropertyChangedEvent
 
 	UpdateMarqueeBoundingBox();
 
+	TangentFlattener.Reset();
 	ActiveTransaction.Reset();
 }
 
@@ -786,6 +787,7 @@ void FCurveEditorTransformTool::OnDragStart()
 	}
 
 	ActiveTransaction = MakeUnique<FScopedTransaction>(TEXT("CurveEditorTransformTool"), LOCTEXT("CurveEditorTransformToolTransaction", "Transform Key(s)"), nullptr);
+	TangentFlattener.Emplace(*CurveEditor);
 
 	CurveEditor->SuppressBoundTransformUpdates(true);
 
@@ -1152,6 +1154,7 @@ void FCurveEditorTransformTool::OnDragEnd()
 
 	CurveEditor->SuppressBoundTransformUpdates(false);
 
+	TangentFlattener.Reset();
 	// This finalizes the transaction
 	ActiveTransaction.Reset();
 
@@ -1250,6 +1253,13 @@ void FCurveEditorTransformTool::ScaleFrom(const FVector2D& InPanelSpaceCenter, c
 		KeyData.LastDraggedKeyPositions = NewKeyPositionScratch;
 
 	}
+
+	// This recalculates the user specified tangents based on the distance moved. We'll flatten the curves depending on the resulting selection's height.
+	if (bInAffectsY)
+	{
+		const bool bIsMirroring = PercentChanged.Y < 0.0;
+		TangentFlattener->ComputeMirroringParallel(*CurveEditor, bIsMirroring);
+	}
 }
 
 void FCurveEditorTransformTool::Tick(const FGeometry& AllottedGeometry, const double InCurrentTime, const float InDeltaTime)
diff --git a/Engine/Plugins/Editor/CurveEditorTools/Source/CurveEditorTools/Private/Tools/CurveEditorTransformTool.h b/Engine/Plugins/Editor/CurveEditorTools/Source/CurveEditorTools/Private/Tools/CurveEditorTransformTool.h
--- a/Engine/Plugins/Editor/CurveEditorTools/Source/CurveEditorTools/Private/Tools/CurveEditorTransformTool.h
+++ b/Engine/Plugins/Editor/CurveEditorTools/Source/CurveEditorTools/Private/Tools/CurveEditorTransformTool.h
@@ -1,13 +1,14 @@
 // Copyright Epic Games, Inc. All Rights Reserved.
 
 #pragma once
+
 #include "CurveEditorTypes.h"
 #include "ICurveEditorToolExtension.h"
 #include "Framework/DelayedDrag.h"
 #include "CurveDataAbstraction.h"
 #include "ScopedTransaction.h"
 #include "CurveEditorSnapMetrics.h"
-
+#include "Misc/Mirror/TangentSelectionFlattener.h"
 #include "CurveEditorTransformTool.generated.h"
 
 class FCurveEditor;
@@ -162,6 +163,8 @@ private:
 
 	/** Set when attempting to move a drag handle. This allows us to tell the difference between a click and a click-drag. */
 	TOptional<FDelayedDrag> DelayedDrag;
+	/** Squishes the tangents vertically depending on how much we move the keys vertically. */
+	TOptional<UE::CurveEditor::FTangentSelectionFlattener> TangentFlattener;
 
 	/** Used to cache selected key data when doing transform operations */
 	struct FKeyData
diff --git a/Engine/Plugins/Editor/CurveEditorTools/Source/CurveEditorTools/Private/Tools/Lattice/CurveEditorLatticeTool.cpp b/Engine/Plugins/Editor/CurveEditorTools/Source/CurveEditorTools/Private/Tools/Lattice/CurveEditorLatticeTool.cpp
--- a/Engine/Plugins/Editor/CurveEditorTools/Source/CurveEditorTools/Private/Tools/Lattice/CurveEditorLatticeTool.cpp
+++ b/Engine/Plugins/Editor/CurveEditorTools/Source/CurveEditorTools/Private/Tools/Lattice/CurveEditorLatticeTool.cpp
@@ -6,9 +6,13 @@
 #include "CurveEditorToolCommands.h"
 #include "Drag/LatticeDragOp_MoveControlPoints.h"
 #include "Internationalization/Text.h"
+#include "Mirror/LatticeEdgeTangentMirrorOp.h"
+#include "Mirror/LatticePointTangentMirrorOp.h"
 #include "Misc/CurvePointSnapper.h"
 #include "Misc/LatticeDrawUtils.h"
 #include "Misc/LatticeUtils.h"
+#include "Misc/VectorMathUtils.h"
+#include "Misc/Mirror/MirrorUtils.h"
 #include "Templates/UnrealTemplate.h"
 #include "UObject/LatticeUndoData.h"
 #include "UObject/Package.h"
@@ -99,24 +103,44 @@ void FCurveEditorLatticeTool::OnPaint(
 	}
 }
 
+FReply FCurveEditorLatticeTool::OnMouseButtonDoubleClick(
+	TSharedRef<SWidget> OwningWidget, const FGeometry& InViewContainerGeometry, const FPointerEvent& InMouseEvent
+	)
+{
+	CancelDrag();
+	const bool bHandled = ComputeHoverStateForLeftClick(InViewContainerGeometry, InMouseEvent,
+		[this](const FLatticeHoverState& HoverState, const FLatticeBounds&)
+		{
+			// Double clicking the center will mirror everything with the midpoint being a horizontal line through center of the grid.
+			if (HoverState.HoveredCell)
+			{
+				const FScopedTransaction Transaction(LOCTEXT("FlipCurves", "Flip curves"));
+				MirrorGrid();
+				return true;
+			}
+			if (HoverState.HoveredEdge)
+			{
+				const FScopedTransaction Transaction(LOCTEXT("FlattenCurves", "Flatten curves"));
+				MoveEdgeToOpposite(*HoverState.HoveredEdge);
+				return true;
+			}
+			return false;
+		});
+	return bHandled
+		// Capturing the mouse is important so OnMouseButtonDown and OnMouseButtonUp are called on us instead of bubbling them up the hierarchy.
+		// Without this, SInteractiveCurveEditorView::OnMouseButtonUp can cause a curve to be clicked and select all keys under it.
+		? FReply::Handled().CaptureMouse(OwningWidget)
+		: FReply::Unhandled();
+}
+	
 FReply FCurveEditorLatticeTool::OnMouseButtonDown(
 	TSharedRef<SWidget> InOwningWidget, const FGeometry& InViewContainerGeometry, const FPointerEvent& InMouseEvent
 	)
 {
 	CancelDrag();
-	
-	if (const TSharedPtr<FCurveEditor> CurveEditor = WeakCurveEditor.Pin()
-		; CurveEditor && LatticeDeformerState.IsSet())
-	{
-		const FLatticeBounds Bounds = ComputeBounds(*CurveEditor);
-		if (Bounds.bIsVisible && InMouseEvent.GetEffectingButton() == EKeys::LeftMouseButton)
+	const bool bHandled = ComputeHoverStateForLeftClick(InViewContainerGeometry, InMouseEvent,
+		[this, &InMouseEvent](const FLatticeHoverState& HoverState, const FLatticeBounds& Bounds)
 		{
-			const TArray<FVector2D> ControlPointsInSlate = ComputeDrawControlPoints(*CurveEditor);
-			const FLatticeHoverState HoverState = ComputeLatticeHoverState(
-				InViewContainerGeometry, InMouseEvent.GetScreenSpacePosition(),
-				ControlPointsInSlate, ComputeDrawControlEdges(ControlPointsInSlate),
-				GetDeformer().NumPointsInWidth(), GetDeformer().NumCells()
-				);
 			if (HoverState.IsHovered())
 			{
 				DelayedDragData.Emplace(
@@ -125,21 +149,22 @@ FReply FCurveEditorLatticeTool::OnMouseButtonDown(
 					Bounds,
 					HoverState
 					);
-				return FReply::Handled().CaptureMouse(InOwningWidget).PreventThrottling();
+				return true;
 			}
-		}
-	}
+			return false;
+		});
 	
-	return FReply::Unhandled();
+	return bHandled ? FReply::Handled().CaptureMouse(InOwningWidget).PreventThrottling() : FReply::Unhandled();
 }
 
+	
 FReply FCurveEditorLatticeTool::OnMouseButtonUp(
 	TSharedRef<SWidget> InOwningWidget, const FGeometry& InViewContainerGeometry, const FPointerEvent& MouseEvent
 	)
 {
-	FReply Reply = FReply::Handled();
 	if (DelayedDragData.IsSet())
 	{
+		FReply Reply = FReply::Handled();
 		if (DelayedDragData->DelayedDrag.IsDragging() && ensure(DelayedDragData->ActiveOperation))
 		{
 			DelayedDragData->ActiveOperation->DragOp->EndDrag(InViewContainerGeometry, MouseEvent);
@@ -169,10 +194,7 @@ FReply FCurveEditorLatticeTool::OnMouseMove(
 		}
 		else if (DelayedDragData->DelayedDrag.AttemptDragStart(InMouseEvent))
 		{
-			DelayedDragData->ActiveOperation.Emplace(
-				CreateDragDrop(DelayedDragData->Bounds, DelayedDragData->HoverState),
-				LOCTEXT("LatticeChangeKeys", "Lattice keys")
-				);
+			StartDragDrop(DelayedDragData->Bounds, DelayedDragData->HoverState);
 			DelayedDragData->ActiveOperation->DragOp->BeginDrag(
 				InViewContainerGeometry, InMouseEvent, DelayedDragData->DelayedDrag.GetInitialPosition()
 				);
@@ -373,46 +395,55 @@ TArray<FLatticeControlEdge> FCurveEditorLatticeTool::ComputeDrawControlEdges(TCo
 	}
 	return ControlEdges;
 }
-	
-TUniquePtr<FLatticeDragOp> FCurveEditorLatticeTool::CreateDragDrop(const FLatticeBounds& InBounds, const FLatticeHoverState& InHoverState)
+
+void FCurveEditorLatticeTool::StartDragDrop(const FLatticeBounds& InBounds, const FLatticeHoverState& InHoverState)
 {
 	const TConstArrayView<FVector2D> ControlPoints = GetDeformer().GetControlPoints();
 	if (InHoverState.HoveredControlPoint)
 	{
 		const int32 ControlPointIndex = *InHoverState.HoveredControlPoint;
-		return MakeUnique<FLatticeDragOp_MoveControlPoints>(
+		TUniquePtr<FLatticeDragOp> Op = MakeUnique<FLatticeDragOp_MoveControlPoints>(
 			WeakCurveEditor,
 			FControlPointArray{ ControlPoints[ControlPointIndex] },
 			InBounds,
 			FLatticeDragOp_MoveControlPoints::FOnControlPointsMoved::CreateLambda(
 				[this, ControlPointIndex](const TConstArrayView<FVector2D>& InNewControlPointPositions)
 				{
-					OnSetControlPoints({ ControlPointIndex }, InNewControlPointPositions);
+					OnMoveControlPoint({ ControlPointIndex }, InNewControlPointPositions);
 				}
 			)
 		);
+		DelayedDragData->ActiveOperation.Emplace(MoveTemp(Op), LOCTEXT("MoveControlPoint", "Move control point"));
+		DelayedDragData->PointDragTangentMirrorOperation = FLatticePointTangentsMirrorOp::MakeFromDraggedControlPoint(
+			ControlPointIndex, GetDeformer(), *WeakCurveEditor.Pin()
+			);
 	}
 
 	if (InHoverState.HoveredEdge)
 	{
-		const auto[Start, End] = GetDeformer().GetEdgeIndices(*InHoverState.HoveredEdge);
-		return MakeUnique<FLatticeDragOp_MoveControlPoints>(
+		const int32 EdgeIndex = *InHoverState.HoveredEdge;
+		const auto[Start, End] = GetDeformer().GetEdgeIndices(EdgeIndex);
+		TUniquePtr<FLatticeDragOp> Op = MakeUnique<FLatticeDragOp_MoveControlPoints>(
 			WeakCurveEditor,
 			FControlPointArray{ ControlPoints[Start], ControlPoints[End] },
 			InBounds,
 			FLatticeDragOp_MoveControlPoints::FOnControlPointsMoved::CreateLambda(
 				[this, Start, End](const TConstArrayView<FVector2D>& InNewControlPointPositions)
 				{
-					OnSetControlPoints({ Start, End }, InNewControlPointPositions);
+					OnMoveEdge({ Start, End }, InNewControlPointPositions);
 				}
 			)
 		);
+		DelayedDragData->ActiveOperation.Emplace(MoveTemp(Op), LOCTEXT("MoveEdge", "Move edge"));
+		DelayedDragData->EdgeDragTangentMirrorOperation = FLatticeEdgeTangentsMirrorOp::MakeFromDraggedEdge(
+			EdgeIndex, GetDeformer(), *WeakCurveEditor.Pin()
+			);
 	}
 
 	if (InHoverState.HoveredCell)
 	{
 		const auto[TopLeft, TopRight, BottomRight, BottomLeft] = GetDeformer().GetCellIndices(*InHoverState.HoveredCell);
-		return MakeUnique<FLatticeDragOp_MoveControlPoints>(
+		TUniquePtr<FLatticeDragOp> Op = MakeUnique<FLatticeDragOp_MoveControlPoints>(
 			WeakCurveEditor,
 			FControlPointArray{ ControlPoints[TopLeft], ControlPoints[TopRight], ControlPoints[BottomRight], ControlPoints[BottomLeft] },
 			InBounds,
@@ -423,10 +454,32 @@ TUniquePtr<FLatticeDragOp> FCurveEditorLatticeTool::CreateDragDrop(const FLattic
 				}
 			)
 		);
+		DelayedDragData->ActiveOperation.Emplace(MoveTemp(Op), LOCTEXT("MoveCell", "Move cell"));
+	}
+}
+
+void FCurveEditorLatticeTool::OnMoveControlPoint(TConstArrayView<int32> InVertexIndices, TConstArrayView<FVector2D> InNewControlPointPositions)
+{
+	check(InVertexIndices.Num() == 1 && InNewControlPointPositions.Num() == 1);
+	
+	if (DelayedDragData->PointDragTangentMirrorOperation)
+	{
+		DelayedDragData->PointDragTangentMirrorOperation->OnMovePoint(InNewControlPointPositions[0]);
+	}
+	
+	OnSetControlPoints(InVertexIndices, InNewControlPointPositions);
+}
+
+void FCurveEditorLatticeTool::OnMoveEdge(TConstArrayView<int32> InVertexIndices, TConstArrayView<FVector2D> InNewControlPointPositions)
+{
+	check(InVertexIndices.Num() == 2 && InNewControlPointPositions.Num() == 2);
+	
+	if (DelayedDragData->EdgeDragTangentMirrorOperation)
+	{
+		DelayedDragData->EdgeDragTangentMirrorOperation->OnMoveEdge(InNewControlPointPositions);
 	}
 
-	checkNoEntry();
-	return nullptr;
+	OnSetControlPoints(InVertexIndices, InNewControlPointPositions);
 }
 
 void FCurveEditorLatticeTool::OnSetControlPoints(TConstArrayView<int32> InControlPointsToUpdate, TConstArrayView<FVector2D> InNewControlPoints)
@@ -487,6 +540,98 @@ void FCurveEditorLatticeTool::OnSetControlPoints(TConstArrayView<int32> InContro
 		}
 	}
 }
+
+void FCurveEditorLatticeTool::MoveEdgeToOpposite(int32 InEdgeIndex)
+{
+	TSharedPtr<FCurveEditor> CurveEditor = WeakCurveEditor.Pin();
+	const FLatticeDeformer2D& Deformer = GetDeformer();
+	// At time of implementation, we only have a single cell. Should that change, adjust this implementation.
+	if (Deformer.NumCells() != 1 || !CurveEditor)
+	{
+		return;
+	}
+
+	const ELatticeEdgeType EdgeType = GetEdgeTypeFromIndexInSingleCellGrid(InEdgeIndex);
+	if (EdgeType != ELatticeEdgeType::Top && EdgeType != ELatticeEdgeType::Bottom)
+	{
+		return;
+	}
+	
+	const ELatticeEdgeType OppositeEdge = GetOppositeEdge(EdgeType);
+	const int32 OppositeEdgeIndex = GetEdgeIndexInSingleCellGrid(OppositeEdge);
+	const auto[MovedEdgeStartVertIdx, MovedEdgeEndVertIdx] = Deformer.GetEdgeIndices(OppositeEdgeIndex);
+	const FLatticeControlEdge& Edge = Deformer.GetControlEdge(InEdgeIndex);
+	const TArray<FVector2D, TInlineAllocator<2>> NewEdge = { Edge.Start(), Edge.End() };
+	
+	// OnSetControlPoints modifies the curves too late - MirrorTangents will already have modified them.
+	const TConstArrayView<FKeyInCurveId> DeformerKeys = Deformer.GetKeyIdsInCell(0);
+	ModifyCurves(DeformerKeys, *CurveEditor); 
+	
+	// Flatten tangents if moving top or bottom edge...
+	TOptional<FLatticeEdgeTangentsMirrorOp> MirrorOp = FLatticeEdgeTangentsMirrorOp::MakeFromDraggedEdge(
+		OppositeEdgeIndex, GetDeformer(), *WeakCurveEditor.Pin()
+		);
+	if (ensure(MirrorOp)) // Should be set because we've checked above that we're moving a top or bottom edge
+	{
+		MirrorOp->OnMoveEdge(NewEdge);
+	}
+	// ... then move the grid
+	OnSetControlPoints({ MovedEdgeStartVertIdx, MovedEdgeEndVertIdx }, NewEdge);
+}
+
+void FCurveEditorLatticeTool::MirrorGrid()
+{
+	TSharedPtr<FCurveEditor> CurveEditor = WeakCurveEditor.Pin();
+	const FLatticeDeformer2D& Deformer = GetDeformer();
+	// At time of implementation, we only have a single cell. Should that change, adjust this implementation.
+	if (Deformer.NumCells() != 1 || !CurveEditor)
+	{
+		return;
+	}
+
+	const TConstArrayView<FKeyInCurveId> DeformerKeys = Deformer.GetKeyIdsInCell(0);
+	ModifyCurves(DeformerKeys, *CurveEditor); // OnSetControlPoints modifies the curves too late - MirrorTangents will already have modified them.
+
+	// 1. Mirror all the keys
+	const TMap<FCurveModelID, TArray<FKeyHandle>> SplitCurves = SplitKeysInCurves(DeformerKeys);
+	CurveEditor::MirrorOnMidpoint(SplitCurves, *CurveEditor, [&CurveEditor](
+		TConstArrayView<FKeyHandle> Keys, const FCurveModelID& CurveId, double Min, double Max, double Midpoint)
+		{
+			CurveEditor::MirrorTangents(Keys, CurveId, Min, Max, Midpoint, *CurveEditor);
+			CurveEditor::MirrorKeyPositions(Keys, CurveId, Min, Max, Midpoint, *CurveEditor);
+		});
+
+	// 2. Recompute lattice grid
+	// This will calculate a new rectangular grid around the points. Instead, we could try to preserve the existing shape... but nobody cares, so lets not implement it.
+	LatticeDeformerState.Reset();
+	UpdateDeformerBounds(CurveEditor.ToSharedRef());
+}
+	
+template <typename TCallback> requires std::is_invocable_r_v<bool, TCallback, const FLatticeHoverState&, const FLatticeBounds&>
+bool FCurveEditorLatticeTool::ComputeHoverStateForLeftClick(
+	const FGeometry& InViewContainerGeometry, const FPointerEvent& InMouseEvent, TCallback&& InCallback
+	)
+{
+	const TSharedPtr<FCurveEditor> CurveEditor = WeakCurveEditor.Pin();
+	if (!CurveEditor || !LatticeDeformerState.IsSet())
+	{
+		return false;
+	}
+	
+	const FLatticeBounds Bounds = ComputeBounds(*CurveEditor);
+	if (Bounds.bIsVisible && InMouseEvent.GetEffectingButton() == EKeys::LeftMouseButton)
+	{
+		const TArray<FVector2D> ControlPointsInSlate = ComputeDrawControlPoints(*CurveEditor);
+		const FLatticeHoverState HoverState = ComputeLatticeHoverState(
+			InViewContainerGeometry, InMouseEvent.GetScreenSpacePosition(),
+			ControlPointsInSlate, ComputeDrawControlEdges(ControlPointsInSlate),
+			GetDeformer().NumPointsInWidth(), GetDeformer().NumCells()
+			);
+		return InCallback(HoverState, Bounds);
+	}
+
+	return false;
+}
 }
 
 #undef LOCTEXT_NAMESPACE
\ No newline at end of file
diff --git a/Engine/Plugins/Editor/CurveEditorTools/Source/CurveEditorTools/Private/Tools/Lattice/CurveEditorLatticeTool.h b/Engine/Plugins/Editor/CurveEditorTools/Source/CurveEditorTools/Private/Tools/Lattice/CurveEditorLatticeTool.h
--- a/Engine/Plugins/Editor/CurveEditorTools/Source/CurveEditorTools/Private/Tools/Lattice/CurveEditorLatticeTool.h
+++ b/Engine/Plugins/Editor/CurveEditorTools/Source/CurveEditorTools/Private/Tools/Lattice/CurveEditorLatticeTool.h
@@ -7,6 +7,8 @@
 #include "ICurveEditorToolExtension.h"
 #include "LatticeControlsDrawData.h"
 #include "LatticeDeformer2D.h"
+#include "Mirror/LatticeEdgeTangentMirrorOp.h"
+#include "Mirror/LatticePointTangentMirrorOp.h"
 #include "Misc/CurveChangeListener.h"
 #include "Misc/LatticeDrawUtils.h"
 #include "ScopedTransaction.h"
@@ -33,6 +35,7 @@ public:
 	//~ Begin ICurveEditorToolExtension Interface
 	virtual void Tick(const FGeometry& InViewContainerGeometry, const double InCurrentTime, const float InDeltaTime) override;
 	virtual void OnPaint(const FPaintArgs& InPaintArgs, const FGeometry& InViewContainerGeometry, const FSlateRect& InMyCullingRect, FSlateWindowElementList& OutDrawElements, int32 InPaintOnLayerId, const FWidgetStyle& InWidgetStyle, bool bParentEnabled) const override;
+	virtual FReply OnMouseButtonDoubleClick(TSharedRef<SWidget> OwningWidget, const FGeometry& InViewContainerGeometry, const FPointerEvent& InMouseEvent) override;
 	virtual FReply OnMouseButtonDown(TSharedRef<SWidget> InOwningWidget, const FGeometry& InViewContainerGeometry, const FPointerEvent& InMouseEvent) override;
 	virtual FReply OnMouseButtonUp(TSharedRef<SWidget> InOwningWidget, const FGeometry& InViewContainerGeometry, const FPointerEvent& MouseEvent) override;
 	virtual FReply OnMouseMove(TSharedRef<SWidget> InOwningWidget, const FGeometry& InViewContainerGeometry, const FPointerEvent& InMouseEvent) override;
@@ -109,6 +112,10 @@ private:
 
 		/** Set if the drag up was started. */
 		TOptional<FDragOpData> ActiveOperation;
+		/** Adjust tangents while dragging on an edge.*/
+		TOptional<FLatticeEdgeTangentsMirrorOp> EdgeDragTangentMirrorOperation;
+		/** Adjusts tangents while dragging on a control point. */
+		TOptional<FLatticePointTangentsMirrorOp> PointDragTangentMirrorOperation;
 		
 		explicit FDelayedDragData(
 			const FVector2D& InInitialPosition, const FKey& InEffectiveKey, const FLatticeBounds& InBounds, const FLatticeHoverState& InHoverState
@@ -141,11 +148,27 @@ private:
 	TArray<FVector2D> ComputeDrawControlPoints(const FCurveEditor& InCurveEditor) const;
 	TArray<FLatticeControlEdge> ComputeDrawControlEdges(TConstArrayView<FVector2D> InControlPoints) const;
 
-	TUniquePtr<FLatticeDragOp> CreateDragDrop(const FLatticeBounds& InBounds, const FLatticeHoverState& InHoverState);
+	/** Starts a new drag drop operation based on whether a control point, edge, or cell is hovered. */
+	void StartDragDrop(const FLatticeBounds& InBounds, const FLatticeHoverState& InHoverState);
+	/** Called when a control point is dragged. */
+	void OnMoveControlPoint(TConstArrayView<int32> InVertexIndices, TConstArrayView<FVector2D> InNewControlPointPositions);
+	/** Called when an edge is dragged. */
+	void OnMoveEdge(TConstArrayView<int32> InVertexIndices, TConstArrayView<FVector2D> InNewControlPointPositions);
 	/** Generic function for updating control points. */
 	void OnSetControlPoints(TConstArrayView<int32> InControlPointsToUpdate, TConstArrayView<FVector2D> InNewControlPoints);
+
+	/** Moves the clicked edge to the opposite edge, effectively flattening the curve sections. */
+	void MoveEdgeToOpposite(int32 InEdgeIndex);
+	/** Mirrors all points in the grid across an imaginary, x-axis aligned edge going through the grid center. Only if the grid consists of a single cell. */
+	void MirrorGrid();
 	
 	const FLatticeDeformer2D& GetDeformer() const { check(LatticeDeformerState); return LatticeDeformerState->Deformer; }
 	FLatticeDeformer2D& GetDeformer() { check(LatticeDeformerState); return LatticeDeformerState->Deformer; }
+
+	/**
+	 * Invokes InCallback if there is a lattice grid shown and the user left clicks somewhere.
+	 * @return False if InCallback was not invoked. Otherwise, whatever InCallback returned. */
+	template<typename TCallback> requires std::is_invocable_r_v<bool, TCallback, const FLatticeHoverState&, const FLatticeBounds&>
+	bool ComputeHoverStateForLeftClick(const FGeometry& InViewContainerGeometry, const FPointerEvent& InMouseEvent, TCallback&& InCallback);
 };
 }
diff --git a/Engine/Plugins/Editor/CurveEditorTools/Source/CurveEditorTools/Private/Tools/Lattice/LatticeDeformer2D.h b/Engine/Plugins/Editor/CurveEditorTools/Source/CurveEditorTools/Private/Tools/Lattice/LatticeDeformer2D.h
--- a/Engine/Plugins/Editor/CurveEditorTools/Source/CurveEditorTools/Private/Tools/Lattice/LatticeDeformer2D.h
+++ b/Engine/Plugins/Editor/CurveEditorTools/Source/CurveEditorTools/Private/Tools/Lattice/LatticeDeformer2D.h
@@ -4,11 +4,10 @@
 
 #include "Containers/Array.h"
 #include "Containers/ContainersFwd.h"
-#include "Curves/KeyHandle.h"
-#include "CurveDataAbstraction.h"
 #include "HAL/Platform.h"
 #include "Math/MathFwd.h"
 #include "Math/Vector2D.h"
+#include "Misc/KeyInCurveId.h"
 #include "Templates/FunctionFwd.h"
 
 namespace UE::CurveEditorTools
@@ -34,19 +33,6 @@ struct FLatticeControlEdge
 	{}
 };
 
-struct FKeyInCurveId
-{
-	FKeyHandle KeyId;
-	FCurveModelID CurveId;
-	
-	friend FArchive& operator<<(FArchive& InArchive, FKeyInCurveId& InCurveData)
-	{
-		InArchive << InCurveData.KeyId;
-		InArchive << InCurveData.CurveId;
-		return InArchive;
-	}
-};
-
 /** Identifies the normalized coordinates a key has in a cell. */
 using FLatticeKeyCoords = FVector2D;
 	
diff --git a/Engine/Plugins/Editor/CurveEditorTools/Source/CurveEditorTools/Private/Tools/Lattice/Mirror/LatticeEdgeTangentMirrorOp.cpp b/Engine/Plugins/Editor/CurveEditorTools/Source/CurveEditorTools/Private/Tools/Lattice/Mirror/LatticeEdgeTangentMirrorOp.cpp
new file mode 100644
--- /dev/null
+++ b/Engine/Plugins/Editor/CurveEditorTools/Source/CurveEditorTools/Private/Tools/Lattice/Mirror/LatticeEdgeTangentMirrorOp.cpp
@@ -0,0 +1,72 @@
+﻿// Copyright Epic Games, Inc. All Rights Reserved.
+
+#include "LatticeEdgeTangentMirrorOp.h"
+
+#include "Tools/Lattice/LatticeDeformer2D.h"
+#include "Tools/Lattice/Misc/LatticeUtils.h"
+#include "MirrorUtils.h"
+#include "Misc/Mirror/TangentMirrorSolver.h"
+#include "Misc/Optional.h"
+
+namespace UE::CurveEditorTools
+{
+namespace MirrorEdgeDetail
+{
+static bool CanUseEdgeForMirroring(int32 InEdgeIndex, const FLatticeDeformer2D& InDeformer)
+{
+	const bool bCorrectEdge = InEdgeIndex == GetEdgeIndexInSingleCellGrid(ELatticeEdgeType::Top)
+		|| InEdgeIndex == GetEdgeIndexInSingleCellGrid(ELatticeEdgeType::Bottom);
+	return bCorrectEdge && InDeformer.NumCells() == 1;
+}
+	
+static void GetDragStartAndMidpointHeights(
+	int32 InEdgeIndex, const FLatticeDeformer2D& InDeformer, EMirroredEdgeType& OutEdgeType, double& OutDragStartHeight, double& OutMidPointHeight
+)
+{
+	constexpr int32 TopIndex = GetEdgeIndexInSingleCellGrid(ELatticeEdgeType::Top);
+	const bool bIsTopEdge = TopIndex == InEdgeIndex;
+	const int32 OppositeEdgeIndex = bIsTopEdge ? GetEdgeIndexInSingleCellGrid(ELatticeEdgeType::Bottom) : TopIndex;
+
+	const auto[StartVertIndex, EndVertIndex] = InDeformer.GetEdgeIndices(InEdgeIndex);
+	const auto[OppositeStartVertIndex, OppositeEndVertIndex] = InDeformer.GetEdgeIndices(OppositeEdgeIndex);
+
+	const TConstArrayView<FVector2D> ControlPoints = InDeformer.GetControlPoints();
+	OutEdgeType = bIsTopEdge ? EMirroredEdgeType::Top : EMirroredEdgeType::Bottom;
+	OutDragStartHeight = ChooseConsistentEdgeVert(OutEdgeType, ControlPoints[StartVertIndex], ControlPoints[EndVertIndex]);
+	OutMidPointHeight = ChooseConsistentEdgeVert(OutEdgeType, ControlPoints[OppositeStartVertIndex], ControlPoints[OppositeEndVertIndex]);
+}
+}
+
+TOptional<FLatticeEdgeTangentsMirrorOp> FLatticeEdgeTangentsMirrorOp::MakeFromDraggedEdge(
+	int32 InEdgeIndex, const FLatticeDeformer2D& InDeformer, const FCurveEditor& InCurveEditor
+	)
+{
+	if (MirrorEdgeDetail::CanUseEdgeForMirroring(InEdgeIndex, InDeformer))
+	{
+		EMirroredEdgeType EdgeType;
+		double DragStartHeight, MidPointHeight;
+		MirrorEdgeDetail::GetDragStartAndMidpointHeights(InEdgeIndex, InDeformer, EdgeType, DragStartHeight, MidPointHeight);
+
+		CurveEditor::FTangentMirrorSolver TangentSolver(DragStartHeight, MidPointHeight);
+		bool bAtLeastOne = false;
+		for (const TPair<FCurveModelID, TArray<FKeyHandle>>& Pair : SplitKeysInCurves(InDeformer.GetKeyIdsInCell(0)))
+		{
+			bAtLeastOne |= TangentSolver.AddTangents(InCurveEditor, Pair.Key, Pair.Value);
+		}
+
+		if (bAtLeastOne)
+		{
+			return FLatticeEdgeTangentsMirrorOp(InCurveEditor, EdgeType, MoveTemp(TangentSolver));
+		}
+	}
+	return {};
+}
+
+void FLatticeEdgeTangentsMirrorOp::OnMoveEdge(TConstArrayView<FVector2D> InNewEdge)
+{
+	if (ensureMsgf(InNewEdge.Num() == 2, TEXT("We expected an edge")))
+	{
+		TangentMirrorSolver.OnMoveEdge(CurveEditor,  ChooseConsistentEdgeVert(EdgeType, InNewEdge[0], InNewEdge[1]));
+	}
+}
+}
\ No newline at end of file
diff --git a/Engine/Plugins/Editor/CurveEditorTools/Source/CurveEditorTools/Private/Tools/Lattice/Mirror/LatticeEdgeTangentMirrorOp.h b/Engine/Plugins/Editor/CurveEditorTools/Source/CurveEditorTools/Private/Tools/Lattice/Mirror/LatticeEdgeTangentMirrorOp.h
new file mode 100644
--- /dev/null
+++ b/Engine/Plugins/Editor/CurveEditorTools/Source/CurveEditorTools/Private/Tools/Lattice/Mirror/LatticeEdgeTangentMirrorOp.h
@@ -0,0 +1,51 @@
+﻿// Copyright Epic Games, Inc. All Rights Reserved.
+
+#pragma once
+
+#include "MirrorUtils.h"
+#include "HAL/Platform.h"
+#include "Tools/Lattice/Misc/LatticeUtils.h"
+#include "Misc/Mirror/TangentMirrorSolver.h"
+
+class FCurveEditor;
+namespace UE::CurveEditor { class FTangentMirrorSolver; }
+template<typename OptionalType> struct TOptional;
+
+/** Implements mirroring functions based off of lattice tool's data. */
+namespace UE::CurveEditorTools
+{
+class FLatticeDeformer2D;
+	
+/**
+ * Provides functionality to drag the top or bottom edge of a lattice grid and mirror the tangents accordingly.
+ * @note This operation only affects the tangents. Key positions must be adjusted separately using the lattice grid algorithm.
+ */
+class FLatticeEdgeTangentsMirrorOp
+{
+public:
+
+	/** @return Constructed instance that will mirror all keys in InDeformer. Returns valid if the edge is a top or bottom edge in a single celled grid. */
+	static TOptional<FLatticeEdgeTangentsMirrorOp> MakeFromDraggedEdge(
+		int32 InEdgeIndex, const FLatticeDeformer2D& InDeformer, const FCurveEditor& InCurveEditor UE_LIFETIMEBOUND
+		);
+
+	/** Call with info about where the new edge is located. */
+	void OnMoveEdge(TConstArrayView<FVector2D> InNewEdge);
+	
+private:
+
+	/** Used to get and set key attributes */
+	const FCurveEditor& CurveEditor;
+	/** Edge type of the dragged edge. Important for determining height of sloped edges. */
+	const EMirroredEdgeType EdgeType;
+	
+	/** Computes the tangents */
+	CurveEditor::FTangentMirrorSolver TangentMirrorSolver;
+	
+	FLatticeEdgeTangentsMirrorOp(const FCurveEditor& InCurveEditor, const EMirroredEdgeType InEdgeType, CurveEditor::FTangentMirrorSolver InSolver)
+		: CurveEditor(InCurveEditor)
+		, EdgeType(InEdgeType)
+		, TangentMirrorSolver(MoveTemp(InSolver))
+	{}
+};
+}
diff --git a/Engine/Plugins/Editor/CurveEditorTools/Source/CurveEditorTools/Private/Tools/Lattice/Mirror/LatticePointTangentMirrorOp.cpp b/Engine/Plugins/Editor/CurveEditorTools/Source/CurveEditorTools/Private/Tools/Lattice/Mirror/LatticePointTangentMirrorOp.cpp
new file mode 100644
--- /dev/null
+++ b/Engine/Plugins/Editor/CurveEditorTools/Source/CurveEditorTools/Private/Tools/Lattice/Mirror/LatticePointTangentMirrorOp.cpp
@@ -0,0 +1,116 @@
+﻿// Copyright Epic Games, Inc. All Rights Reserved.
+
+#include "LatticePointTangentMirrorOp.h"
+
+#include "MirrorUtils.h"
+#include "Misc/Optional.h"
+#include "Tools/Lattice/LatticeDeformer2D.h"
+#include "Tools/Lattice/Misc/LatticeUtils.h"
+
+namespace UE::CurveEditorTools
+{
+namespace MirrorPointDetail
+{
+enum class EEdgeSide { Left, Right };
+	
+static void GetDragStartAndMidpointHeights(
+	int32 InControlPointIndex, const FLatticeDeformer2D& InDeformer, EEdgeSide& OutEdgeType, double& OutDragStartHeight, double& OutMidPointHeight
+)
+{
+	constexpr int32 TopIndex = GetEdgeIndexInSingleCellGrid(ELatticeEdgeType::Top);
+	constexpr int32 BottomIndex = GetEdgeIndexInSingleCellGrid(ELatticeEdgeType::Bottom);
+	const auto[TopStartVertIndex, TopEndVertIndex] = InDeformer.GetEdgeIndices(TopIndex);
+	const auto[BottomStartVertIndex, BottomEndVertIndex] = InDeformer.GetEdgeIndices(BottomIndex);
+
+	const auto[TopLeft, TopRight, BottomRight, BottomLeft] = InDeformer.GetCellIndices(0);
+	const bool bIsLeft = InControlPointIndex == TopLeft || InControlPointIndex == BottomLeft;
+	OutEdgeType = bIsLeft ? EEdgeSide::Left : EEdgeSide::Right;
+
+	const TConstArrayView<FVector2D> ControlPoints = InDeformer.GetControlPoints();
+	const bool bDraggedIsTopEdge = InControlPointIndex == TopStartVertIndex || InControlPointIndex == TopEndVertIndex;
+	OutDragStartHeight = InDeformer.GetControlPoints()[InControlPointIndex].Y;
+	OutMidPointHeight = bDraggedIsTopEdge
+		? ChooseConsistentEdgeVert(EMirroredEdgeType::Bottom, ControlPoints[BottomStartVertIndex], ControlPoints[BottomEndVertIndex])
+		: ChooseConsistentEdgeVert(EMirroredEdgeType::Top, ControlPoints[TopStartVertIndex], ControlPoints[TopEndVertIndex]);
+}
+}
+	
+TOptional<FLatticePointTangentsMirrorOp> FLatticePointTangentsMirrorOp::MakeFromDraggedControlPoint(
+	int32 InControlPointIndex, const FLatticeDeformer2D& InDeformer, const FCurveEditor& InCurveEditor
+	)
+{
+	using namespace MirrorPointDetail;
+	
+	if (InDeformer.NumCells() == 1)
+	{
+		EEdgeSide EdgeSide;
+		double DragStartHeight, MidPointHeight;
+		GetDragStartAndMidpointHeights(InControlPointIndex, InDeformer, EdgeSide, DragStartHeight, MidPointHeight);
+		
+		const auto[TopLeftIdx, TopRightIdx, BottomRightIdx, BottomLeftIdx] = InDeformer.GetCellIndices(0);
+		const TConstArrayView<FVector2D> ControlPoints = InDeformer.GetControlPoints();
+		const double MinX = FMath::Min(ControlPoints[TopLeftIdx].X, ControlPoints[BottomLeftIdx].X);
+		const double MaxX = FMath::Min(ControlPoints[TopRightIdx].X, ControlPoints[BottomRightIdx].X);
+		const double Width = MaxX - MinX;
+		TArray<FKeyPosition> Positions; 
+		TMap<FCurveModelID, TArray<float>> CurveToFalloffValues;
+		
+		bool bAtLeastOne = false;
+		CurveEditor::FTangentMirrorSolver TangentSolver(DragStartHeight, MidPointHeight);
+		for (const TPair<FCurveModelID, TArray<FKeyHandle>>& Pair : SplitKeysInCurves(InDeformer.GetKeyIdsInCell(0)))
+		{
+			const FCurveModelID CurveId = Pair.Key;
+			// We're going to compute falloff values for each key: between 0 and 1.
+			// The closer the key is to the other side of the lattice grid, the closer the falloff value is to 0.
+			// The closer the key is to the control point, the closer the falloff value is to 1. 
+			const auto ComputeFalloff = [&InCurveEditor, EdgeSide, MinX, Width, &Positions, &CurveToFalloffValues, &CurveId]
+				(const CurveEditor::FMirrorableTangentInfo& InTangentInfo)
+				{
+					const FCurveModel* Curve = InCurveEditor.FindCurve(CurveId);
+					const int32 NumKeys = InTangentInfo.MirrorableKeys.Num();
+					Positions.SetNumUninitialized(NumKeys, EAllowShrinking::No);
+					Curve->GetKeyPositions(InTangentInfo.MirrorableKeys, Positions);
+
+					TArray<float>& FalloffValues = CurveToFalloffValues.Add(CurveId);
+					FalloffValues.SetNumUninitialized(NumKeys);
+					for (int32 Index = 0; Index < NumKeys; ++Index)
+					{
+						// If key is on the left = 0, and on right = 1.
+						const float PercentFromLeft = (Positions[Index].InputValue - MinX) / Width;
+						// If control point is on the left		-> key on left side -> falloff = 1,		key on right side -> falloff = 0
+						// If control point is on the right		-> key on left side -> falloff = 0,		key on right side -> falloff = 1
+						const float Falloff = EdgeSide == EEdgeSide::Left ? 1.f - PercentFromLeft : PercentFromLeft;
+						FalloffValues[Index] = Falloff;
+					}
+				};
+			bAtLeastOne |= TangentSolver.AddTangents(InCurveEditor, CurveId, Pair.Value, ComputeFalloff);
+		}
+		
+		if (bAtLeastOne)
+		{
+			return FLatticePointTangentsMirrorOp(InCurveEditor, MoveTemp(TangentSolver), MoveTemp(CurveToFalloffValues));
+		}
+	}
+	return {};
+}
+
+void FLatticePointTangentsMirrorOp::OnMovePoint(const FVector2D& InControlPoint)
+{
+	for (TPair<FCurveModelID, CurveEditor::FCurveTangentMirrorData>& Pair : TangentMirrorSolver.CurveData)
+	{
+		const FCurveModelID& CurveId = Pair.Key;
+		CurveEditor::FCurveTangentMirrorData& MirrorData = Pair.Value;
+		const TArray<float>& FalloffValues = CurveToFalloffValues[CurveId];
+		RecomputeMirroringParallel(CurveEditor, CurveId, MirrorData, InControlPoint.Y,
+			// Flatten the tangents more if their key is closer to the side of the dragged control point, and don't flatten at all on opposite side.
+			// Effectively, this interpolates the interpolated tangents...
+			[&FalloffValues, &MirrorData](int32 KeyIndex, const FVector2D& InterpolatedTangents)
+			{
+				const float Alpha = FalloffValues[KeyIndex];
+				const float OneMinus = 1.f - Alpha;
+				// Weighted average of interpolated and original tangent value. Effectively, this is a linear falloff.
+				return Alpha * InterpolatedTangents + OneMinus * MirrorData.TangentSolver.InitialValues[KeyIndex];
+			});
+	}
+}
+}
diff --git a/Engine/Plugins/Editor/CurveEditorTools/Source/CurveEditorTools/Private/Tools/Lattice/Mirror/LatticePointTangentMirrorOp.h b/Engine/Plugins/Editor/CurveEditorTools/Source/CurveEditorTools/Private/Tools/Lattice/Mirror/LatticePointTangentMirrorOp.h
new file mode 100644
--- /dev/null
+++ b/Engine/Plugins/Editor/CurveEditorTools/Source/CurveEditorTools/Private/Tools/Lattice/Mirror/LatticePointTangentMirrorOp.h
@@ -0,0 +1,60 @@
+﻿// Copyright Epic Games, Inc. All Rights Reserved.
+
+#pragma once
+
+#include "HAL/Platform.h"
+#include "Misc/Mirror/TangentMirrorSolver.h"
+
+class FCurveEditor;
+namespace UE::CurveEditor { class FTangentMirrorSolver; }
+template<typename OptionalType> struct TOptional;
+
+/** Implements mirroring functions based off of lattice tool's data. */
+namespace UE::CurveEditorTools
+{
+class FLatticeDeformer2D;
+
+/**
+ * Provides functionality to drag a control point of a lattice grid and mirror the tangents accordingly.
+ *
+ * This implementation differs from dragging an edge: a falloff is applied to all tangents. Tangents are flattened the full amount the closer they are
+ * to the dragged control point. Tangents of keys on the other side of the lattice grid are flattened less.
+ * 
+ * @note This operation only affects the tangents. Key positions must be adjusted separately using the lattice grid algorithm.
+ */
+class FLatticePointTangentsMirrorOp
+{
+public:
+
+	/** @return Constructed instance that will mirror all keys in InDeformer. Returns valid if in a single celled grid. */
+	static TOptional<FLatticePointTangentsMirrorOp> MakeFromDraggedControlPoint(
+		int32 InControlPointIndex, const FLatticeDeformer2D& InDeformer, const FCurveEditor& InCurveEditor UE_LIFETIMEBOUND
+		);
+
+	/** Call with info about where the new control point is located. */
+	void OnMovePoint(const FVector2D& InControlPoint);
+
+private:
+
+	/** Used to get and set key attributes */
+	const FCurveEditor& CurveEditor;
+
+	/** Computes the tangents */
+	CurveEditor::FTangentMirrorSolver TangentMirrorSolver;
+
+	/**
+	 * Binds the curves to an array containing falloff values for each key. The array is equal lenght as FCurveTangentMirrorData::KeyHandles.
+	 * The falloff values are between 1 and 0: 1 the closer the key is to the dragged control point and 0 on the other side.
+	 * This is used so tangents closer the control point are flattened more.
+	 */
+	TMap<FCurveModelID, TArray<float>> CurveToFalloffValues;
+
+	FLatticePointTangentsMirrorOp(
+		const FCurveEditor& InCurveEditor, CurveEditor::FTangentMirrorSolver InSolver, TMap<FCurveModelID, TArray<float>> CurveToFalloffValues
+		)
+		: CurveEditor(InCurveEditor)
+		, TangentMirrorSolver(MoveTemp(InSolver))
+		, CurveToFalloffValues(MoveTemp(CurveToFalloffValues))
+	{}
+};
+}
\ No newline at end of file
diff --git a/Engine/Plugins/Editor/CurveEditorTools/Source/CurveEditorTools/Private/Tools/Lattice/Mirror/MirrorUtils.h b/Engine/Plugins/Editor/CurveEditorTools/Source/CurveEditorTools/Private/Tools/Lattice/Mirror/MirrorUtils.h
new file mode 100644
--- /dev/null
+++ b/Engine/Plugins/Editor/CurveEditorTools/Source/CurveEditorTools/Private/Tools/Lattice/Mirror/MirrorUtils.h
@@ -0,0 +1,27 @@
+﻿// Copyright Epic Games, Inc. All Rights Reserved.
+
+#pragma once
+
+#include "Math/Vector.h"
+#include "Misc/AssertionMacros.h"
+
+namespace UE::CurveEditorTools
+{
+enum class EMirroredEdgeType { Top, Bottom };
+
+/** Consistenly chooses the same vertex on the top or bottom edge ASSUMING that the vertices are always transformed uniformly (e.g. move both up by x units, etc.) */
+inline double ChooseConsistentEdgeVert(EMirroredEdgeType EdgeType, const FVector2D& Vert1, const FVector2D& Vert2)
+{
+	// In several places, we need to consistently choose vert 1 or vert 2 to take the height from.
+	// During the edge drag operation the vertices are moved up and down by the same amounts.
+	// For top edge we choose the higher, and for bottom edge the lower vertex on the edge, which guarantees we consistently choose the same vertex.
+	// Choosing top and bottom vertex, respectively, is done so with non-parallel edges the tangents are interpolated as slowly as possible, i.e.
+	// you need to move the mouse more. 
+	switch (EdgeType)
+	{
+	case EMirroredEdgeType::Top: return FMath::Max(Vert1.Y, Vert2.Y);
+	case EMirroredEdgeType::Bottom: return FMath::Min(Vert1.Y, Vert2.Y);
+	default: checkNoEntry(); return 0;
+	}
+}
+}
diff --git a/Engine/Plugins/Editor/CurveEditorTools/Source/CurveEditorTools/Private/Tools/Lattice/Misc/LatticeUtils.h b/Engine/Plugins/Editor/CurveEditorTools/Source/CurveEditorTools/Private/Tools/Lattice/Misc/LatticeUtils.h
--- a/Engine/Plugins/Editor/CurveEditorTools/Source/CurveEditorTools/Private/Tools/Lattice/Misc/LatticeUtils.h
+++ b/Engine/Plugins/Editor/CurveEditorTools/Source/CurveEditorTools/Private/Tools/Lattice/Misc/LatticeUtils.h
@@ -77,4 +77,46 @@ TArray<FVector2D> PlaceLatticeGridOverNewKeyPositions(const FLatticeDeformer2D&
 
 /** Adds InCurveEditor's selection to InDeformer. */
 void AddSelectionToDeformer(FLatticeDeformer2D& InDeformer, const FCurveEditor& InCurveEditor);
+
+enum class ELatticeEdgeType : uint8
+{
+	Top,
+	Bottom,
+	Left,
+	Right,
+};
+/** @return Assuming the lattice grid consists of a single cell, gets the specified edge. */
+constexpr int32 GetEdgeIndexInSingleCellGrid(ELatticeEdgeType Type);
+/** @return Assuming the lattice grid consists of a single cell, gets type of edge based off of its index. */
+constexpr ELatticeEdgeType GetEdgeTypeFromIndexInSingleCellGrid(int32 InEdgeIndex);
+	
+/** @return The opposite edge of the passed in edge */
+constexpr ELatticeEdgeType GetOppositeEdge(ELatticeEdgeType Type);
+}
+
+namespace UE::CurveEditorTools
+{
+constexpr int32 GetEdgeIndexInSingleCellGrid(ELatticeEdgeType Type)
+{
+	// The enum entries are intentionally sorted for this op
+	return static_cast<int32>(Type);
+}
+
+constexpr ELatticeEdgeType GetEdgeTypeFromIndexInSingleCellGrid(int32 InEdgeIndex)
+{
+	// The enum entries are intentionally sorted for this op
+	return static_cast<ELatticeEdgeType>(InEdgeIndex);
+}
+
+constexpr ELatticeEdgeType GetOppositeEdge(ELatticeEdgeType Type)
+{
+	switch (Type)
+	{
+	case ELatticeEdgeType::Top: return ELatticeEdgeType::Bottom;
+	case ELatticeEdgeType::Bottom: return ELatticeEdgeType::Top;
+	case ELatticeEdgeType::Left: return ELatticeEdgeType::Right;
+	case ELatticeEdgeType::Right: return ELatticeEdgeType::Left;
+	default: return ELatticeEdgeType::Top;
+	}
+}
 }
diff --git a/Engine/Source/Editor/CurveEditor/Private/Misc/Mirror/MirrorUtils.cpp b/Engine/Source/Editor/CurveEditor/Private/Misc/Mirror/MirrorUtils.cpp
new file mode 100644
--- /dev/null
+++ b/Engine/Source/Editor/CurveEditor/Private/Misc/Mirror/MirrorUtils.cpp
@@ -0,0 +1,123 @@
+﻿// Copyright Epic Games, Inc. All Rights Reserved.
+
+#include "Misc/Mirror/MirrorUtils.h"
+
+#include "Containers/Array.h"
+#include "CurveEditor.h"
+#include "Misc/Mirror/PositionMirrorSolver.h"
+#include "Misc/Mirror/TangentMirrorSolver.h"
+
+namespace UE::CurveEditor
+{
+void MirrorTangentsOnMidpoint(const TMap<FCurveModelID, TArray<FKeyHandle>>& InKeysToMirror, FCurveEditor& InCurveEditor)
+{
+	for (const TPair<FCurveModelID, TArray<FKeyHandle>>& Pair : InKeysToMirror)
+	{
+		if (const FCurveModel* Model = InCurveEditor.FindCurve(Pair.Key))
+		{
+			const auto[Min, Max] = FindMinMaxHeight(Pair.Value, *Model);
+			const double Midpoint = Min + (Max - Min) * 0.5;
+			MirrorTangents(Pair.Value, Pair.Key, Min, Max, Midpoint, InCurveEditor);
+		}
+	}
+}
+
+FCurveBounds FindMinMaxHeight(TConstArrayView<FKeyHandle> InKeys, const FCurveModel& InCurveModel)
+{
+	if (!ensure(!InKeys.IsEmpty()))
+	{
+		return { 0, 0 };
+	}
+
+	TArray<FKeyPosition> Positions;
+	Positions.SetNumUninitialized(InKeys.Num());
+	InCurveModel.GetKeyPositions(InKeys, Positions);
+
+	double Min = Positions[0].OutputValue;
+	double Max = Positions[0].OutputValue;
+	for (int32 Index = 1; Index < InKeys.Num(); ++Index)
+	{
+		Min = FMath::Min(Positions[Index].OutputValue, Min);
+		Max = FMath::Max(Positions[Index].OutputValue, Max);
+	}
+	
+	return FCurveBounds{ Min, Max };
+}
+
+void MirrorKeyPositions(
+	TConstArrayView<FKeyHandle> InKeysToMirror, const FCurveModelID& InCurveId,
+	double InBottomHeight, double InTopHeight, double InMirrorHeight, FCurveEditor& InCurveEditor)
+{
+	const FCurveModel* CurveModel = InCurveEditor.FindCurve(InCurveId);
+	if (!ensure(CurveModel))
+	{
+		return;
+	}
+	
+	// Subdivide the keys into 2 groups: those above and those below the midpoint.
+	TArray<FKeyHandle> OldAboveMirrorEdge, OldBelowMirrorEdge;
+	for (const FKeyHandle& KeyHandle : InKeysToMirror)
+	{
+		FKeyPosition KeyPosition;
+		CurveModel->GetKeyPositions(TConstArrayView<FKeyHandle>(&KeyHandle, 1), TArrayView<FKeyPosition>(&KeyPosition, 1));
+		if (KeyPosition.OutputValue > InMirrorHeight)
+		{
+			OldAboveMirrorEdge.Add(KeyHandle);
+		}
+		else
+		{
+			OldBelowMirrorEdge.Add(KeyHandle);
+		}
+	}
+	
+	// Then, mirror the bottom to the top, and the top to the bottom.
+	// Move imaginary edge going through top control point opposite the mirror edge...
+	FPositionMirrorSolver TopToBottom(InCurveEditor, InTopHeight, InMirrorHeight);
+	TopToBottom.AddKeyPositions(InCurveId, OldAboveMirrorEdge);
+	TopToBottom.OnMoveEdge(InBottomHeight);
+	
+	// ... and do the same for the bottom
+	FPositionMirrorSolver BottomToTop(InCurveEditor, InBottomHeight, InMirrorHeight);
+	BottomToTop.AddKeyPositions(InCurveId , OldBelowMirrorEdge);
+	BottomToTop.OnMoveEdge(InTopHeight);
+}
+
+void MirrorTangents(
+	TConstArrayView<FKeyHandle> InKeysToMirror, const FCurveModelID& InCurveId,
+	double InBottomHeight, double InTopHeight, double InMirrorHeight, FCurveEditor& InCurveEditor
+	)
+{
+	const FCurveModel* CurveModel = InCurveEditor.FindCurve(InCurveId);
+	if (!ensure(CurveModel))
+	{
+		return;
+	}
+	
+	// Subdivide the keys into 2 groups: those above and those below the midpoint.
+	TArray<FKeyHandle> OldAboveMirrorEdge, OldBelowMirrorEdge;
+	for (const FKeyHandle& KeyHandle : InKeysToMirror)
+	{
+		FKeyPosition KeyPosition;
+		CurveModel->GetKeyPositions(TConstArrayView<FKeyHandle>(&KeyHandle, 1), TArrayView<FKeyPosition>(&KeyPosition, 1));
+		if (KeyPosition.OutputValue > InMirrorHeight)
+		{
+			OldAboveMirrorEdge.Add(KeyHandle);
+		}
+		else
+		{
+			OldBelowMirrorEdge.Add(KeyHandle);
+		}
+	}
+	
+	// Then, mirror the bottom to the top, and the top to the bottom.
+	// Move imaginary edge going through top control point opposite the mirror edge...
+	FTangentMirrorSolver TopToBottom(InTopHeight, InMirrorHeight);
+	TopToBottom.AddTangents(InCurveEditor, InCurveId, OldAboveMirrorEdge);
+	TopToBottom.OnMoveEdge(InCurveEditor, InBottomHeight);
+	
+	// ... and do the same for the bottom
+	FTangentMirrorSolver BottomToTop(InBottomHeight, InMirrorHeight);
+	BottomToTop.AddTangents(InCurveEditor, InCurveId , OldBelowMirrorEdge);
+	BottomToTop.OnMoveEdge(InCurveEditor, InTopHeight);
+}
+}
\ No newline at end of file
diff --git a/Engine/Source/Editor/CurveEditor/Private/Misc/Mirror/MultiCurveMirrorUtils.cpp b/Engine/Source/Editor/CurveEditor/Private/Misc/Mirror/MultiCurveMirrorUtils.cpp
new file mode 100644
--- /dev/null
+++ b/Engine/Source/Editor/CurveEditor/Private/Misc/Mirror/MultiCurveMirrorUtils.cpp
@@ -0,0 +1,69 @@
+﻿// Copyright Epic Games, Inc. All Rights Reserved.
+
+#include "Misc/Mirror/MultiCurveMirrorUtils.h"
+
+#include "CurveEditor.h"
+#include "CurveModel.h"
+
+namespace UE::CurveEditor
+{
+FMirrorableTangentInfo FilterMirrorableTangents(const FCurveEditor& InCurveEditor, const FCurveModelID& InCurveId, TConstArrayView<FKeyHandle> InKeys)
+{
+	FMirrorableTangentInfo Result;
+	TArray<FKeyHandle>& OutMirroredKeys = Result.MirrorableKeys;
+	TArray<FKeyAttributes>& OutInitialAttributes = Result.InitialAttributes;
+	TArray<FVector2D>& OutTangents = Result.Tangents;
+	TArray<double>& OutHeights = Result.KeyHeights;
+	
+	const FCurveModel* CurveModel = InCurveEditor.FindCurve(InCurveId);
+	if (!ensure(CurveModel))
+	{
+		return Result;
+	}
+	
+	const int32 NumKeys = InKeys.Num();
+	OutMirroredKeys.Reserve(NumKeys);
+	OutTangents.Reserve(NumKeys);
+	OutHeights.Reserve(NumKeys);
+	
+	for (int32 Index = 0; Index < NumKeys; ++Index)
+	{
+		const FKeyHandle& KeyHandle = InKeys[Index];
+		if (!CurveModel)
+		{
+			continue;
+		}
+
+		FKeyAttributes KeyAttribute;
+		CurveModel->GetKeyAttributesExcludingAutoComputed(
+			TConstArrayView<FKeyHandle>(&KeyHandle, 1), TArrayView<FKeyAttributes>(&KeyAttribute, 1)
+			);
+
+		const bool bHasValidMode = KeyAttribute.HasTangentMode()
+			&& (KeyAttribute.GetTangentMode() == RCTM_User || KeyAttribute.GetTangentMode() == RCTM_Break);
+		if (!bHasValidMode)
+		{
+			continue;
+		}
+
+		OutMirroredKeys.Add(KeyHandle);
+		if (KeyAttribute.HasArriveTangent() && KeyAttribute.HasLeaveTangent())
+		{
+			OutTangents.Add({ KeyAttribute.GetArriveTangent(), KeyAttribute.GetLeaveTangent() });
+		}
+		else if (KeyAttribute.HasArriveTangent())
+		{
+			OutTangents.Add({ KeyAttribute.GetArriveTangent(), 0.0 });
+		}
+
+		FKeyPosition KeyPosition;
+		CurveModel->GetKeyPositions(TConstArrayView<FKeyHandle>(&KeyHandle, 1), TArrayView<FKeyPosition>(&KeyPosition, 1));
+		OutHeights.Add(KeyPosition.OutputValue);
+	}
+	
+	OutInitialAttributes.SetNumUninitialized(OutMirroredKeys.Num());
+	CurveModel->GetKeyAttributes(OutMirroredKeys, OutInitialAttributes);
+
+	return Result;
+}
+}
diff --git a/Engine/Source/Editor/CurveEditor/Private/Misc/Mirror/PositionMirrorSolver.cpp b/Engine/Source/Editor/CurveEditor/Private/Misc/Mirror/PositionMirrorSolver.cpp
new file mode 100644
--- /dev/null
+++ b/Engine/Source/Editor/CurveEditor/Private/Misc/Mirror/PositionMirrorSolver.cpp
@@ -0,0 +1,79 @@
+﻿// Copyright Epic Games, Inc. All Rights Reserved.
+
+#include "Misc/Mirror/PositionMirrorSolver.h"
+
+#include "Containers/Array.h"
+#include "Containers/Set.h"
+#include "CurveEditor.h"
+
+namespace UE::CurveEditor
+{
+FPositionMirrorSolver::FPositionMirrorSolver(FCurveEditor& InCurveEditor, double InStartY, double InMiddlePointY)
+	: CurveEditor(InCurveEditor)
+	, StartY(InStartY)
+	, MiddlePointY(InMiddlePointY)
+{}
+
+bool FPositionMirrorSolver::AddKeyPositions(
+	const FCurveModelID& InCurveId, TArray<FKeyHandle> InKeys, TArray<FKeyPosition> InPositions
+	)
+{
+	const bool bIsValid = !InKeys.IsEmpty() && (InKeys.Num() == InPositions.Num() || InPositions.IsEmpty());  
+	if (!bIsValid)
+	{
+		return false;
+	}
+	
+	if (InPositions.IsEmpty())
+	{
+		const FCurveModel* CurveModel = CurveEditor.FindCurve(InCurveId);
+		if (!CurveModel)
+		{
+			return false;
+		}
+		InPositions.SetNumUninitialized(InKeys.Num());
+		CurveModel->GetKeyPositions(InKeys, InPositions);
+	}
+
+	TArray<double> Heights;
+	Algo::Transform(InPositions, Heights, [](const FKeyPosition& InKeyPosition){ return InKeyPosition.OutputValue; });
+
+	TArray<double> HeightsRelativeToMiddlePoint = Heights;
+	for (double& Height : HeightsRelativeToMiddlePoint)
+	{
+		// TMirrorSolver compute MiddlePointY + MirrorAlpha * Height.
+		// MirrorAlpha = 1 at start edge position, MirrorAlpha = -1 on opposite side (perfectly mirrored).
+		Height -= MiddlePointY;
+	}
+	
+	AllCurveData.Add(
+		InCurveId,
+		FCachedCurveData(
+			TUniformMirrorSolver<double>(StartY, MiddlePointY, MoveTemp(HeightsRelativeToMiddlePoint), MoveTemp(Heights), MiddlePointY),
+			MoveTemp(InKeys), MoveTemp(InPositions)
+			)
+	);
+	return true;
+}
+
+void FPositionMirrorSolver::OnMoveEdge(double InDraggedEdgeHeight)
+{
+	for (TPair<FCurveModelID, FCachedCurveData>& TangentDataPair : AllCurveData)
+	{
+		FCurveModel* CurveModel = CurveEditor.FindCurve(TangentDataPair.Key);
+		if (!CurveModel)
+		{
+			continue;
+		}
+
+		FCachedCurveData& CurveData = TangentDataPair.Value;
+		CurveData.TangentSolver.ComputeMirroringParallel(InDraggedEdgeHeight, [this, &CurveData](int32 KeyIndex, double NewValue)
+		{
+			FKeyPosition& Position = CurveData.PositionsToSet[KeyIndex];
+			Position.OutputValue = NewValue;
+		});
+		
+		CurveModel->SetKeyPositions(CurveData.KeyHandles, CurveData.PositionsToSet);
+	}
+}
+}
\ No newline at end of file
diff --git a/Engine/Source/Editor/CurveEditor/Private/Misc/Mirror/TangentSelectionFlattener.cpp b/Engine/Source/Editor/CurveEditor/Private/Misc/Mirror/TangentSelectionFlattener.cpp
new file mode 100644
--- /dev/null
+++ b/Engine/Source/Editor/CurveEditor/Private/Misc/Mirror/TangentSelectionFlattener.cpp
@@ -0,0 +1,126 @@
+﻿// Copyright Epic Games, Inc. All Rights Reserved.
+
+#include "Misc/Mirror/TangentSelectionFlattener.h"
+
+#include "CurveEditor.h"
+#include "Misc/Mirror/MirrorUtils.h"
+
+namespace UE::CurveEditor
+{
+namespace TangentFlattenDetail
+{
+static void FilterByMirrorableTangents(
+	const TConstArrayView<FKeyHandle> KeysInCurves, const FCurveModelID& InCurveId, FCurveEditor& InCurveEditor,
+	TArray<FKeyHandle>& OutMirroredKeys, TArray<FKeyAttributes>& OutInitialAttributes, TArray<FVector2D>& OutTangents, TArray<double>& OutHeights
+	)
+	{
+		const FCurveModel* CurveModel = InCurveEditor.FindCurve(InCurveId);
+		if (!ensure(CurveModel))
+		{
+			return;
+		}
+	
+		const int32 NumKeys = KeysInCurves.Num();
+		OutMirroredKeys.Reserve(NumKeys);
+		OutTangents.Reserve(NumKeys);
+		OutHeights.Reserve(NumKeys);
+	
+		for (int32 Index = 0; Index < NumKeys; ++Index)
+		{
+			const FKeyHandle& KeyHandle = KeysInCurves[Index];
+			if (!CurveModel)
+			{
+				continue;
+			}
+
+			FKeyAttributes KeyAttribute;
+			CurveModel->GetKeyAttributesExcludingAutoComputed(
+				TConstArrayView<FKeyHandle>(&KeyHandle, 1), TArrayView<FKeyAttributes>(&KeyAttribute, 1)
+				);
+
+			const bool bHasValidMode = KeyAttribute.HasTangentMode()
+				&& (KeyAttribute.GetTangentMode() == RCTM_User || KeyAttribute.GetTangentMode() == RCTM_Break);
+			if (!bHasValidMode)
+			{
+				continue;
+			}
+
+			OutMirroredKeys.Add(KeyHandle);
+			if (KeyAttribute.HasArriveTangent() && KeyAttribute.HasLeaveTangent())
+			{
+				OutTangents.Add({ KeyAttribute.GetArriveTangent(), KeyAttribute.GetLeaveTangent() });
+			}
+			else if (KeyAttribute.HasArriveTangent())
+			{
+				OutTangents.Add({ KeyAttribute.GetArriveTangent(), 0.0 });
+			}
+
+			FKeyPosition KeyPosition;
+			CurveModel->GetKeyPositions(TConstArrayView<FKeyHandle>(&KeyHandle, 1), TArrayView<FKeyPosition>(&KeyPosition, 1));
+			OutHeights.Add(KeyPosition.OutputValue);
+		}
+	
+		OutInitialAttributes.SetNumUninitialized(OutMirroredKeys.Num());
+		CurveModel->GetKeyAttributes(OutMirroredKeys, OutInitialAttributes);
+}
+}
+
+bool FTangentSelectionFlattener::ResetFromSelection(const FCurveEditor& InCurveEditor)
+{
+	CachedCurveData.Empty(CachedCurveData.Num());
+	for (const TPair<FCurveModelID, FKeyHandleSet>& Pair : InCurveEditor.GetSelection().GetAll())
+	{
+		AddTangents(InCurveEditor, Pair.Key, Pair.Value.AsArray());
+	}
+	return !CachedCurveData.IsEmpty();
+}
+
+bool FTangentSelectionFlattener::AddTangents(const FCurveEditor& InCurveEditor, const FCurveModelID& InCurveId, TConstArrayView<FKeyHandle> InKeys)
+{
+	const FCurveModel* CurveModel = InCurveEditor.FindCurve(InCurveId);
+	if (!CurveModel)
+	{
+		return false;
+	}
+		
+	FMirrorableTangentInfo TangentInfo = FilterMirrorableTangents(InCurveEditor, InCurveId, InKeys);
+	if (!TangentInfo)
+	{
+		return false;
+	}
+		
+	const auto[Min, Max] = FindMinMaxHeight(TangentInfo.MirrorableKeys, *CurveModel);
+	const double SelectionHeight = Max - Min;
+	if (FMath::IsNearlyZero(SelectionHeight))
+	{
+		return false;
+	}
+
+	// We're going to have 2 fake edges: the moved "edge" is the height different of min and max values in selection. The 2nd is simply zero.
+	// Effectively, we'll squish the tangents based on how much the selection is squished vertically. 0 height means tangent is 0.
+	constexpr double MidpointEdgeHeight = 0;
+	CachedCurveData.Add(InCurveId, FCurveTangentMirrorData(MoveTemp(TangentInfo), SelectionHeight, MidpointEdgeHeight));
+	return true;
+}
+
+void FTangentSelectionFlattener::ComputeMirroringParallel(const FCurveEditor& InCurveEditor, bool bTopHasCrossedBottomEdge)
+{
+	// If the top edge has crossed the bottom edge since we were initialized, we need to mirror. In that case Alpha is in range [-1, 0].
+	const double Sign = bTopHasCrossedBottomEdge ? -1 : 1;
+	
+	for (TPair<FCurveModelID, FCurveTangentMirrorData>& Pair : CachedCurveData)
+	{
+		FCurveModel* CurveModel = InCurveEditor.FindCurve(Pair.Key);
+		if (!CurveModel)
+		{
+			continue;
+		}
+		
+		FCurveTangentMirrorData& CurveData = Pair.Value;
+		// Like described above, the tangents are squished as much as the selection height difference is squished. 
+		const auto[Min, Max] = FindMinMaxHeight(CurveData.KeyHandles, *CurveModel);
+		const double SelectionHeight = Max - Min;
+		RecomputeMirroringParallel(InCurveEditor, Pair.Key, CurveData, SelectionHeight * Sign);
+	}
+}
+}
diff --git a/Engine/Source/Editor/CurveEditor/Public/CurveEditor.h b/Engine/Source/Editor/CurveEditor/Public/CurveEditor.h
--- a/Engine/Source/Editor/CurveEditor/Public/CurveEditor.h
+++ b/Engine/Source/Editor/CurveEditor/Public/CurveEditor.h
@@ -317,6 +317,8 @@ public:
 	const TMap<FCurveModelID, TUniquePtr<FCurveModel>>& GetCurves() const;
 
 	FCurveEditorSelection& GetSelection() { return Selection; }
+	const FCurveEditorSelection& GetSelection() const { return Selection; }
+	
 	/**
 	 * Generate a utility struct for snapping values
 	 */
diff --git a/Engine/Source/Editor/CurveEditor/Public/Misc/CurveEditorUtils.h b/Engine/Source/Editor/CurveEditor/Public/Misc/CurveEditorUtils.h
new file mode 100644
--- /dev/null
+++ b/Engine/Source/Editor/CurveEditor/Public/Misc/CurveEditorUtils.h
@@ -0,0 +1,20 @@
+﻿// Copyright Epic Games, Inc. All Rights Reserved.
+
+#pragma once
+
+#include "Containers/Array.h"
+#include "CurveEditor.h"
+#include "CurveModel.h"
+
+namespace UE::CurveEditor
+{
+/** Modifies all referenced curves. */
+inline void ModifyCurves(TConstArrayView<FCurveModelID> InCurves, FCurveEditor& InCurveEditor)
+{
+	for (const FCurveModelID& CurveId : InCurves)
+	{
+		FCurveModel* Model = InCurveEditor.FindCurve(CurveId);
+		Model->Modify();
+	}
+}
+}
diff --git a/Engine/Source/Editor/CurveEditor/Public/Misc/Mirror/MirrorUtils.h b/Engine/Source/Editor/CurveEditor/Public/Misc/Mirror/MirrorUtils.h
new file mode 100644
--- /dev/null
+++ b/Engine/Source/Editor/CurveEditor/Public/Misc/Mirror/MirrorUtils.h
@@ -0,0 +1,97 @@
+﻿// Copyright Epic Games, Inc. All Rights Reserved.
+
+#pragma once
+
+#include "CurveEditor.h"
+#include "Containers/Array.h"
+#include "Containers/Map.h"
+
+class FCurveModel;
+class FCurveEditor;
+struct FCurveModelID;
+template<typename OptionalType> struct TOptional;
+
+namespace UE::CurveEditor
+{
+/**
+ * Finds the mid-point of every curve's passed in key handles and mirrors all points on the imaginary line going through it.
+ * @param InKeysToMirror Curve to the keys on it to mirror
+ * @param InCurveEditor The editor to which the keys belong
+ */
+template<typename TCallback, typename TMapAllocator = FDefaultSetAllocator, typename TArrayAllocator = FDefaultAllocator>
+requires std::is_invocable_v<TCallback, TConstArrayView<FKeyHandle>, const FCurveModelID&, double /*Min*/, double /*Max*/, double /*Midpoint*/>
+void MirrorOnMidpoint(
+	const TMap<FCurveModelID, TArray<FKeyHandle, TArrayAllocator>, TMapAllocator>& InKeysToMirror, FCurveEditor& InCurveEditor, TCallback&& InProcess
+	);
+
+struct FCurveBounds
+{
+	/** All keys are above this. */
+	double Min;
+	/** All keys are below this */
+	double Max;
+};
+/** @return Find the min and max height values. */
+CURVEEDITOR_API FCurveBounds FindMinMaxHeight(
+	TConstArrayView<FKeyHandle> InKeys, const FCurveModel& InCurveModel
+	);
+
+/**
+ * Recomputes the key position OutputValues as if the points were mirrored accross an imaginary x-axis aligned mirror edge.
+ * The mirror Edge is positioned between the top and bottom edge.
+ *
+ * @note This function only recalculates key positions. Tangents must be set separately.
+ * 
+ * @param InKeysToMirror Keys located between BottomHeight and TopHeight, whose tangents will be recomputed.
+ * @param InCurveId The curve the keys belong to
+ * @param InBottomHeight The height of the bottom edge.
+ * @param InTopHeight The height of the top edge.
+ * @param InMirrorHeight The height of the imaginary x-axis-aligned mirror edge across which tangents are mirrored.
+ * @param InCurveEditor Reference to the curve editor used to set key attributes.
+ */
+CURVEEDITOR_API void MirrorKeyPositions(
+	TConstArrayView<FKeyHandle> InKeysToMirror, const FCurveModelID& InCurveId,
+	double InBottomHeight, double InTopHeight, double InMirrorHeight,
+	FCurveEditor& InCurveEditor
+	);
+	
+/**
+ * Recomputes the tangents of specified keys as if the points were mirrored across an imaginary x-axis aligned mirror edge.
+ * The mirror edge is positioned between the top and bottom edges.
+ * 
+ * @note This function only recalculates tangents. Key positions must be set separately.
+ * 
+ * @param InKeysToMirror Keys located between BottomHeight and TopHeight, whose tangents will be recomputed.
+ * @param InCurveId The curve the keys belong to
+ * @param InBottomHeight The height of the bottom edge.
+ * @param InTopHeight The height of the top edge.
+ * @param InMirrorHeight The height of the imaginary x-axis-aligned mirror edge across which tangents are mirrored.
+ * @param InCurveEditor Reference to the curve editor used to set key attributes.
+ */
+CURVEEDITOR_API void MirrorTangents(
+	TConstArrayView<FKeyHandle> InKeysToMirror, const FCurveModelID& InCurveId,
+	double InBottomHeight, double InTopHeight, double InMirrorHeight,
+	FCurveEditor& InCurveEditor
+	);
+}
+
+namespace UE::CurveEditor
+{
+template<typename TCallback, typename TMapAllocator, typename TArrayAllocator>
+requires std::is_invocable_v<TCallback, TConstArrayView<FKeyHandle>, const FCurveModelID&, double /*Min*/, double /*Max*/, double /*Midpoint*/>
+void MirrorOnMidpoint(
+	const TMap<FCurveModelID, TArray<FKeyHandle, TArrayAllocator>, TMapAllocator>& InKeysToMirror, FCurveEditor& InCurveEditor, TCallback&& InProcess
+)	
+{
+	for (const TPair<FCurveModelID, TArray<FKeyHandle>>& Pair : InKeysToMirror)
+	{
+		if (const FCurveModel* Model = InCurveEditor.FindCurve(Pair.Key))
+		{
+			const auto[Min, Max] = FindMinMaxHeight(Pair.Value, *Model);
+			const double Midpoint = Min + (Max - Min) * 0.5;
+			InProcess(Pair.Value, Pair.Key, Min, Max, Midpoint);
+		}
+	}
+}
+}
+
diff --git a/Engine/Source/Editor/CurveEditor/Public/Misc/Mirror/MultiCurveMirrorUtils.h b/Engine/Source/Editor/CurveEditor/Public/Misc/Mirror/MultiCurveMirrorUtils.h
new file mode 100644
--- /dev/null
+++ b/Engine/Source/Editor/CurveEditor/Public/Misc/Mirror/MultiCurveMirrorUtils.h
@@ -0,0 +1,130 @@
+﻿// Copyright Epic Games, Inc. All Rights Reserved.
+
+#pragma once
+
+#include "Containers/Array.h"
+#include "CurveDataAbstraction.h"
+#include "CurveEditor.h"
+#include "CurveEditorTypes.h"
+#include "CurveModel.h"
+#include "UniformMirrorSolver.h"
+
+namespace UE::CurveEditor
+{
+struct FMirrorableTangentInfo
+{
+	/** Keys that can be mirrored, i.e. have RCTM_User or RCTM_Break mode. */
+	TArray<FKeyHandle> MirrorableKeys;
+	/** The attributes corresponding to each of MirrorableKeys. */
+	TArray<FKeyAttributes> InitialAttributes;
+	/**
+	 * The tangents corresponding to each of MirrorableKeys.
+	 * X is arrive tangent, Y is leave tangent. Copied values of the corresponding InitialAttributes indices. If there is no leave tangent, Y is 0.
+	 */
+	TArray<FVector2D> Tangents;
+	/** The FKeyPosition::OutputValue corresponding to each of MirrableKeys. */
+	TArray<double> KeyHeights;
+
+	bool HasData() const { return !MirrorableKeys.IsEmpty(); }
+	operator bool() const { return HasData(); }
+};
+
+/** @return Information you can use to FCurveTangentMirrorData. */
+CURVEEDITOR_API FMirrorableTangentInfo FilterMirrorableTangents(
+	const FCurveEditor& InCurveEditor, const FCurveModelID& InCurveId, TConstArrayView<FKeyHandle> InKeys
+	);
+	
+/** Holds data about mirroring keys in a single curve. */
+struct FCurveTangentMirrorData
+{
+	/**
+	 * This will mirror the tangents of KeyHandles. The X component corresponds to ArriveTangent and Y to LeaveTangent.
+	 * TMirrorSolver::InitialValues and TMirrorSolver::InitialKeyHeights correspond to KeyHandles.
+	 */
+	const TUniformMirrorSolver<FVector2D> TangentSolver;
+	/** The keys to interpolate. Only tangents that are user specified are referenced here. */
+	const TArray<FKeyHandle> KeyHandles;
+	/** Same length as KeyHandles.*/
+	TArray<FKeyAttributes> AttributesToSet;
+
+	explicit FCurveTangentMirrorData(TUniformMirrorSolver<FVector2D> InTangentSolver, TArray<FKeyHandle> InKeys, TArray<FKeyAttributes> InInitialAttributes)
+		: TangentSolver(MoveTemp(InTangentSolver))
+		, KeyHandles(MoveTemp(InKeys))
+		, AttributesToSet(MoveTemp(InInitialAttributes))
+	{
+		check(TangentSolver.NumValues() == KeyHandles.Num() && KeyHandles.Num() == AttributesToSet.Num());
+		ensureMsgf(TangentSolver.NumValues() != 0, TEXT("Pointless construction"));
+	}
+	
+	explicit FCurveTangentMirrorData(
+		FMirrorableTangentInfo InInfo, double InStartHeight, double InMirrorMidpoint, const FVector2D& InMidpointOffset = FVector2D::ZeroVector
+		)
+		: FCurveTangentMirrorData(
+			TUniformMirrorSolver<FVector2D>(InStartHeight, InMirrorMidpoint, MoveTemp(InInfo.Tangents), MoveTemp(InInfo.KeyHeights), InMidpointOffset),
+			MoveTemp(InInfo.MirrorableKeys), MoveTemp(InInfo.InitialAttributes)
+		)
+	{}
+};
+
+	
+/** Mirrors the tangents stores in InCurveData. */
+void RecomputeMirroringParallel(
+	const FCurveEditor& InCurveEditor, const FCurveModelID& InCurveId, FCurveTangentMirrorData& InCurveData, double InMirrorValue
+	);
+/**
+ * Mirrors the tangents stores in InCurveData.
+ * This version accepts a callback for post-processing, which is useful for e.g. applying a falloff to the interpolated values.
+ */
+template<typename TCallback> requires std::is_invocable_r_v<FVector2D, TCallback, int32, const FVector2D&>
+void RecomputeMirroringParallel(
+	const FCurveEditor& InCurveEditor, const FCurveModelID& InCurveId, FCurveTangentMirrorData& InCurveData, double InMirrorValue,
+	TCallback&& InPostProcessTangents
+	);
+}
+
+namespace UE::CurveEditor
+{
+inline void RecomputeMirroringParallel(
+	const FCurveEditor& InCurveEditor, const FCurveModelID& InCurveId, FCurveTangentMirrorData& InCurveData, double InMirrorValue
+	)
+{
+	RecomputeMirroringParallel(InCurveEditor, InCurveId, InCurveData, InMirrorValue,
+		[](int32, const FVector2D& Interpolated)
+		{
+			return Interpolated;
+		});
+}
+
+/** Mirrors the tangents stores in InCurveData. */
+template<typename TCallback> requires std::is_invocable_r_v<FVector2D, TCallback, int32/*KeyIndex*/, const FVector2D&/*InterpolatedTangents*/>
+void RecomputeMirroringParallel(
+	const FCurveEditor& InCurveEditor, const FCurveModelID& InCurveId, FCurveTangentMirrorData& InCurveData, double InMirrorValue,
+	TCallback&& InPostProcessTangents
+	)
+{
+	FCurveModel* CurveModel = InCurveEditor.FindCurve(InCurveId);
+	if (!CurveModel)
+	{
+		return;
+	}
+
+	InCurveData.TangentSolver.ComputeMirroringParallel(InMirrorValue, [&InCurveData, &InPostProcessTangents](int32 KeyIndex, const FVector2D& InterpolatedTangents)
+	{
+		// KeyIndex is an index to the Tangents and Heights arrays we passed to TangentSolver's constructor.
+		// KeysToMirror was constructed in such a way that it coincides with those arrays. 
+		FKeyAttributes& Attributes = InCurveData.AttributesToSet[KeyIndex];
+
+		const FVector2D Tangents = InPostProcessTangents(KeyIndex, InterpolatedTangents);
+		if (Attributes.HasArriveTangent())
+		{
+			Attributes.SetArriveTangent(Tangents.X);
+		}
+		if (Attributes.HasLeaveTangent())
+		{
+			Attributes.SetLeaveTangent(Tangents.Y);
+		}
+	});
+
+	CurveModel->SetKeyAttributes(InCurveData.KeyHandles, InCurveData.AttributesToSet);
+}
+}
\ No newline at end of file
diff --git a/Engine/Source/Editor/CurveEditor/Public/Misc/Mirror/PositionMirrorSolver.h b/Engine/Source/Editor/CurveEditor/Public/Misc/Mirror/PositionMirrorSolver.h
new file mode 100644
--- /dev/null
+++ b/Engine/Source/Editor/CurveEditor/Public/Misc/Mirror/PositionMirrorSolver.h
@@ -0,0 +1,92 @@
+﻿// Copyright Epic Games, Inc. All Rights Reserved.
+
+#pragma once
+
+#include "Containers/Array.h"
+#include "CurveDataAbstraction.h"
+#include "CurveEditorTypes.h"
+#include "UniformMirrorSolver.h"
+
+class FCurveEditor;
+template<typename OptionalType> struct TOptional;
+
+namespace UE::CurveEditor
+{
+/**
+ * This class facilitates the movement of a dragged edge and recalculates the key heights for all specified keys.
+ * The heights are linearly interpolated, ensuring that if the dragged edge crosses to the opposite side of the mirror edge,
+ * the keys will flip to the other side.
+ * 
+ * Inputs:
+ *  - Height of the dragged edge aligned to the x-axis
+ *  - Height of the x-axis-aligned mirror edge
+ *  - Keys originally located between the dragged and mirror edges
+ *
+ * Example:
+ * Given a grid with a single key at a height 100, moving the top edge:
+ *  - An equal distance past the bottom edge will result in a height of -100 for the key.
+ *  - Halfway between the starting position and the bottom edge will result in a height of 50 for the key.
+ * 
+ * @note This class only recalculates key positions. To adjust the tangents of the keys use e.g. FTangentMirrorSolver.
+ */
+class CURVEEDITOR_API FPositionMirrorSolver
+{
+public:
+	
+	/**
+	 * InStartY and InMiddlePointY define two imaginary, x-axis aligned lines.
+	 * All keys added must be between or on these two lines for the mirroring to work as expected.
+	 * 
+	 * @param InCurveEditor Used to get and set tangent values
+	 * @param InStartY The height of the line that is moved. It determines how the tangents are interpolated.
+	 * @param InMiddlePointY The height of the line on which the points are moved.
+	 */
+	explicit FPositionMirrorSolver(
+		FCurveEditor& InCurveEditor UE_LIFETIMEBOUND,
+		double InStartY, double InMiddlePointY
+		);
+
+	/**
+	 * Adds key positions that are to be mirrored.
+	 *
+	 * InStartY and InMiddlePointY define two imaginary, x-axis aligned lines.
+	 * All keys must be between or on these two lines for the mirroring to work as expected.
+	 * 
+	 * @param InCurveId The curve that all keys lie on
+	 * @param InKeys The keys whose tangents to recompute
+	 * @param InPositions The positions of the keys, if you have them already (optimization). If not, pass in empty and the key positions are determined.
+	 * @return Whether any of the keys can be interpolated (only user specified tangents can be interpolated).
+	 */
+	bool AddKeyPositions(const FCurveModelID& InCurveId, TArray<FKeyHandle> InKeys, TArray<FKeyPosition> InPositions = {});
+	
+	/** Recomputes the tangents angles and updates the key attributes. */
+	void OnMoveEdge(double InDraggedEdgeHeight);
+	
+private:
+	
+	/** Used to set tangents angles. */
+	FCurveEditor& CurveEditor;
+
+	/** Height of the dragged line. */
+	const double StartY;
+	/** Height of the mirroring line. */
+	const double MiddlePointY;
+	
+	struct FCachedCurveData
+	{
+		const TUniformMirrorSolver<double> TangentSolver;
+		
+		const TArray<FKeyHandle> KeyHandles;
+		TArray<FKeyPosition> PositionsToSet;
+
+		explicit FCachedCurveData(TUniformMirrorSolver<double> InTangentSolver, TArray<FKeyHandle> InKeys, TArray<FKeyPosition> InPositions)
+			: TangentSolver(MoveTemp(InTangentSolver))
+			, KeyHandles(MoveTemp(InKeys))
+			, PositionsToSet(MoveTemp(InPositions))
+		{}
+	};
+	/** This memory is reuse across OnMoveEdge to hold the updated tangents. Caching the memory speeds up the recomputation. */
+	TMap<FCurveModelID, FCachedCurveData> AllCurveData;
+};
+}
+
diff --git a/Engine/Source/Editor/CurveEditor/Public/Misc/Mirror/TangentMirrorSolver.h b/Engine/Source/Editor/CurveEditor/Public/Misc/Mirror/TangentMirrorSolver.h
new file mode 100644
--- /dev/null
+++ b/Engine/Source/Editor/CurveEditor/Public/Misc/Mirror/TangentMirrorSolver.h
@@ -0,0 +1,120 @@
+﻿// Copyright Epic Games, Inc. All Rights Reserved.
+
+#pragma once
+
+#include "Containers/Map.h"
+#include "CurveEditor.h"
+#include "CurveEditorTypes.h"
+#include "MultiCurveMirrorUtils.h"
+
+class FCurveEditor;
+template<typename OptionalType> struct TOptional;
+
+namespace UE::CurveEditor
+{
+/**
+ * Given keys that are between an imaginary, x-axis aligned mirror edge and dragged edge, this class facilitates recalculates the tangents for all
+ * specified keys. The tangents are linearly interpolated, ensuring that if the dragged edge crosses to the opposite side of the mirror edge,
+ * the tangents will flip, resulting in mirrored curves.
+ * 
+ * Inputs:
+ *  - Height of the dragged edge aligned to the x-axis
+ *  - Height of the x-axis-aligned mirror edge
+ *  - Keys originally located between the dragged and mirror edges
+ *
+ * Example: Given a grid with a single key at a 45-degree angle, moving the top edge:
+ *  - An equal distance past the bottom edge will result in a tangent of -45 degrees for the key.
+ *  - Halfway between the starting position and the bottom edge will result in a tangent of 22.5 degrees for the key.
+ * 
+ * @note This class only recalculates tangents. To adjust the positions of the keys use e.g. FPositionMirrorSolver.
+ */
+class FTangentMirrorSolver
+{
+public:
+	
+	/** Height of the dragged line. */
+	const double StartY;
+	/** Height of the mirroring line. */
+	const double MiddlePointY;
+	
+	/** This memory is reused across OnMoveEdge to hold the updated tangents. Caching the memory speeds up the re-computation. */
+	TMap<FCurveModelID, FCurveTangentMirrorData> CurveData;
+	
+	/**
+	 * InStartY and InMiddlePointY define two imaginary, x-axis aligned lines.
+	 * All keys added must be between or on these two lines for the mirroring to work as expected.
+	 * 
+	 * @param InStartY The height of the line that is moved. It determines how the tangents are interpolated.
+	 * @param InMiddlePointY The height of the line on which the points are moved.
+	 */
+	explicit FTangentMirrorSolver(
+		double InStartY, double InMiddlePointY
+		);
+
+	/**
+	 * Adds tangents that are be mirrored.
+	 *
+	 * InStartY and InMiddlePointY define two imaginary, x-axis aligned lines.
+	 * All keys must be between or on these two lines for the mirroring to work as expected.
+	 * 
+	 * @param InCurveEditor Used to get and set tangent values
+	 * @param InCurveId The curve that all keys lie on
+	 * @param InKeys The keys whose tangents to recompute
+	 * @return Whether any of the keys can be interpolated (only user specified tangents can be interpolated).
+	 */
+	bool AddTangents(const FCurveEditor& InCurveEditor, const FCurveModelID& InCurveId, TConstArrayView<FKeyHandle> InKeys);
+	
+	/**
+	 * Version that allows you to do additional processing for the tangents.
+	 * 
+	 * For example,  you may want to compute falloff values for each key. In that case, you probably want to manually iterate CurveData and use the
+	 * version of RecomputeMirroringParallel that allows you to interpolate tangent values further.
+	 */
+	template<typename TProcessCallback> requires std::is_invocable_v<TProcessCallback, const FMirrorableTangentInfo& /*TangentInfo*/>
+	bool AddTangents(
+		const FCurveEditor& InCurveEditor, const FCurveModelID& InCurveId, TConstArrayView<FKeyHandle> InKeys, TProcessCallback&& InProcessTangentInfo
+		);
+	
+	/** Recomputes the tangents angles and updates the key attributes. */
+	void OnMoveEdge(const FCurveEditor& InCurveEditor, double InDraggedEdgeHeight);
+};
+}
+
+namespace UE::CurveEditor
+{
+inline FTangentMirrorSolver::FTangentMirrorSolver(double InStartY, double InMiddlePointY)
+	: StartY(InStartY)
+	, MiddlePointY(InMiddlePointY)
+{}
+
+inline bool FTangentMirrorSolver::AddTangents(const FCurveEditor& InCurveEditor, const FCurveModelID& InCurveId, TConstArrayView<FKeyHandle> InKeys)
+{
+	return AddTangents(InCurveEditor, InCurveId, InKeys, [](const FMirrorableTangentInfo&){});
+}
+	
+template <typename TProcessCallback> requires std::is_invocable_v<TProcessCallback, const FMirrorableTangentInfo&>
+bool FTangentMirrorSolver::AddTangents(
+	const FCurveEditor& InCurveEditor, const FCurveModelID& InCurveId, TConstArrayView<FKeyHandle> InKeys, TProcessCallback&& InProcessTangentInfo
+	)
+{
+	FMirrorableTangentInfo TangentInfo = FilterMirrorableTangents(InCurveEditor, InCurveId, InKeys);
+	if (TangentInfo)
+	{
+		InProcessTangentInfo(TangentInfo);
+		CurveData.Add(
+			InCurveId,
+			FCurveTangentMirrorData(MoveTemp(TangentInfo), StartY, MiddlePointY)
+		);
+		return true;
+	}
+	return false;
+}
+	
+inline void FTangentMirrorSolver::OnMoveEdge(const FCurveEditor& InCurveEditor, double InDraggedEdgeHeight)
+{
+	for (TPair<FCurveModelID, FCurveTangentMirrorData>& Pair : CurveData)
+	{
+		RecomputeMirroringParallel(InCurveEditor, Pair.Key, Pair.Value, InDraggedEdgeHeight);
+	}
+}
+}
\ No newline at end of file
diff --git a/Engine/Source/Editor/CurveEditor/Public/Misc/Mirror/TangentSelectionFlattener.h b/Engine/Source/Editor/CurveEditor/Public/Misc/Mirror/TangentSelectionFlattener.h
new file mode 100644
--- /dev/null
+++ b/Engine/Source/Editor/CurveEditor/Public/Misc/Mirror/TangentSelectionFlattener.h
@@ -0,0 +1,45 @@
+﻿// Copyright Epic Games, Inc. All Rights Reserved.
+
+#pragma once
+
+#include "Containers/Map.h"
+#include "CurveEditor.h"
+#include "MultiCurveMirrorUtils.h"
+
+class FCurveEditor;
+
+namespace UE::CurveEditor
+{
+/**
+ * Given some keys, this class squishes the curves based on how much the selection is squished vertically.
+ * The squishing of the curves is achieved by interpolating the tangents to 0.
+ * 
+ * This is useful if you have a selection of keys, apply some operation that will move the keys, and want to flatten the curves.
+ * Example use case is e.g. tweening.
+ */
+class CURVEEDITOR_API FTangentSelectionFlattener
+{
+public:
+
+	FTangentSelectionFlattener() = default;
+	explicit FTangentSelectionFlattener(const FCurveEditor& InCurveEditor) { ResetFromSelection(InCurveEditor); }
+
+	/** Inits the tangent data from the curve editor's selection. */
+	bool ResetFromSelection(const FCurveEditor& InCurveEditor);
+
+	/** @return Whether any interpolatable tangents were added (only tangents that are RCTM_User or RCTM_Break). */
+	bool AddTangents(const FCurveEditor& InCurveEditor, const FCurveModelID& InCurveId, TConstArrayView<FKeyHandle> InKeys);
+
+	/**
+	 * Call in response to the keys having moved. Recomputes user specified tangents and flattens them.
+	 * @param InCurveEditor Used to update key positions
+	 * @param bTopHasCrossedBottomEdge If since this was initialized, the top edge has moved across the bottom edge.
+	 */
+	void ComputeMirroringParallel(const FCurveEditor& InCurveEditor, bool bTopHasCrossedBottomEdge = false);
+
+private:
+	
+	TMap<FCurveModelID, FCurveTangentMirrorData> CachedCurveData;
+};
+}
+
diff --git a/Engine/Source/Editor/CurveEditor/Public/Misc/Mirror/UniformMirrorSolver.h b/Engine/Source/Editor/CurveEditor/Public/Misc/Mirror/UniformMirrorSolver.h
new file mode 100644
--- /dev/null
+++ b/Engine/Source/Editor/CurveEditor/Public/Misc/Mirror/UniformMirrorSolver.h
@@ -0,0 +1,79 @@
+﻿// Copyright Epic Games, Inc. All Rights Reserved.
+
+#pragma once
+
+#include "Async/ParallelFor.h"
+#include "Containers/Array.h"
+#include "HAL/Platform.h"
+
+#include <type_traits>
+
+namespace UE::CurveEditor
+{
+/**
+ * This class computes linear interpolation for values involving the movement of edges in a rectangular selection of keys. While best used for
+ * tangent angles, it could be applied to any key attributes that can be interpolated linearly.
+ *
+ * For example:
+ * - When the top edge is moved downward, the interpolated values decrease linearly to 0 as the edge approaches 
+ *   the bottom edge. If the top edge moves past the bottom edge (the "midpoint"), the values invert, 
+ *   transitioning past zero into the opposite sign.
+ * - Similarly, when the bottom edge is moved upward toward the top edge, the same interpolation and inversion logic applies.
+ * 
+ * This class handles the described interpolation logic in a generic manner. It is agnostic to the type or source 
+ * of the values being interpolated and how the computed values are applied. Its sole responsibility is to perform 
+ * linear interpolation based on the edge movement.
+ *
+ * @param TNumeric Must support operator+(TNumeric, TNumeric) itself and operator*(TNumeric,double). E.g. float, double, FVector2D, FVector, etc.
+ */
+template<typename TNumeric> 
+class TUniformMirrorSolver
+{
+public:
+	
+	explicit TUniformMirrorSolver(
+		double StartY, double MiddlePointY, TArray<TNumeric> InInterpolatedValues, TArray<double> InKeyHeights, TNumeric InMiddlePointBaseValue = TNumeric(0)
+		)
+		: StartY(StartY)
+		, MiddlePointY(MiddlePointY)
+		, InitialValues(MoveTemp(InInterpolatedValues))
+		, InitialKeyHeights(MoveTemp(InKeyHeights))
+		, MiddlePointBaseValue(InMiddlePointBaseValue)
+	{
+		check(InitialValues.Num() == InitialKeyHeights.Num());
+	}
+
+	/**
+	 * Recomputes the values in response to the edge being moved.
+	 * @param InNewHeight The new height of the moved edge
+	 * @param InProcessTangent Processes the new tangent angles (in radians).
+	 */
+	template<typename TCallback> requires std::is_invocable_v<TCallback, int32 /*InIndex*/, const TNumeric& /*NewValue*/>
+	void ComputeMirroringParallel(double InNewHeight, TCallback&& InProcessTangent) const
+	{
+		const int32 Num = InitialValues.Num(); 
+		const double Alpha = (InNewHeight - MiddlePointY) / (StartY - MiddlePointY);
+		ParallelFor(Num, [this, Alpha, &InProcessTangent](int32 Index)
+		{
+			const TNumeric InitialTangent = InitialValues[Index];
+			const TNumeric NewValue = InitialTangent * Alpha + MiddlePointBaseValue;
+			InProcessTangent(Index, NewValue);
+		});
+	}
+
+	int32 NumValues() const { return InitialValues.Num(); }
+
+	/** The Y component at which the drag started. */
+	const double StartY;
+	/** The Y component at which the tangents reach 0 slope. */
+	const double MiddlePointY;
+	
+	/** The initial tangent angles in radians. Indices coincide with InitialKeyHeights. */
+	const TArray<TNumeric> InitialValues;
+	/** The initial heights of all keys. Indices coincide with InitialKeyHeights. */
+	const TArray<double> InitialKeyHeights;
+	/** The value that should be at MiddlePointY. */
+	const TNumeric MiddlePointBaseValue;
+};
+}
+
