Curve Editor: Lattice tool will now recompute lattice grid when you switch tools and then switch back to lattice tool.

diff --git a/Engine/Plugins/Editor/CurveEditorTools/Source/CurveEditorTools/Private/Tools/Lattice/CurveEditorLatticeTool.cpp b/Engine/Plugins/Editor/CurveEditorTools/Source/CurveEditorTools/Private/Tools/Lattice/CurveEditorLatticeTool.cpp
--- a/Engine/Plugins/Editor/CurveEditorTools/Source/CurveEditorTools/Private/Tools/Lattice/CurveEditorLatticeTool.cpp
+++ b/Engine/Plugins/Editor/CurveEditorTools/Source/CurveEditorTools/Private/Tools/Lattice/CurveEditorLatticeTool.cpp
@@ -226,7 +226,8 @@ void FCurveEditorLatticeTool::OnToolActivated()
 
 void FCurveEditorLatticeTool::OnToolDeactivated()
 {
-	// Don't clear the data - if the user activates after deactivating without changing the selected keys, we'll show them what they last had
+	// Clear data. We'll recompute a new grid next time the user activate the tool
+	LatticeDeformerState.Reset();
 }
 
 void FCurveEditorLatticeTool::BindCommands(TSharedRef<FUICommandList> CommandBindings)
