Curve editor: Lattice tool will now recompute its grid when modified externally. It no longer tries to preserver the last shape because there are too many ways to modify the shape (tween, mmb move, inserting keys with control rig, move handle, to name a few).

diff --git a/Engine/Plugins/Editor/CurveEditorTools/Source/CurveEditorTools/Private/Tools/Lattice/CurveEditorLatticeTool.cpp b/Engine/Plugins/Editor/CurveEditorTools/Source/CurveEditorTools/Private/Tools/Lattice/CurveEditorLatticeTool.cpp
--- a/Engine/Plugins/Editor/CurveEditorTools/Source/CurveEditorTools/Private/Tools/Lattice/CurveEditorLatticeTool.cpp
+++ b/Engine/Plugins/Editor/CurveEditorTools/Source/CurveEditorTools/Private/Tools/Lattice/CurveEditorLatticeTool.cpp
@@ -69,7 +69,7 @@ void FCurveEditorLatticeTool::Tick(const FGeometry& InViewContainerGeometry, con
 	if (LatticeDeformerState && LatticeDeformerState->bHasRequestedRefresh)
 	{
 		LatticeDeformerState->bHasRequestedRefresh = false;
-		AdjustLatticeGridOrRecompute(CurveEditor.ToSharedRef());
+		RecomputeLatticeGrid(CurveEditor.ToSharedRef());
 	}
 	else
 	{
@@ -306,36 +306,15 @@ void FCurveEditorLatticeTool::OnCurvesModified()
 	LatticeDeformerState->bHasRequestedRefresh = !LatticeDeformerState->bIsModifyingCurves;
 }
 
-void FCurveEditorLatticeTool::AdjustLatticeGridOrRecompute(const TSharedRef<FCurveEditor>& InCurveEditor)
+void FCurveEditorLatticeTool::RecomputeLatticeGrid(const TSharedRef<FCurveEditor>& InCurveEditor)
 {
 	checkSlow(LatticeDeformerState);
 	
-	// If the selection has changed completely, then refitting the pre-existing lattice makes no sense.
-	const bool bHasSelectionChanged = SelectionSerialNumber != InCurveEditor->Selection.GetSerialNumber();
-	if (bHasSelectionChanged)
-	{
-		// This will calculate a new rectangular grid around the points.
-		// Instead, we could try to preserve the existing shape by just moving it over the moved keys, provided it still fits.
-		LatticeDeformerState.Reset();
-		UpdateDeformerBounds(InCurveEditor);
-	}
-	else
-	{
-		if (GUndo)
-		{
-			UndoObject->Modify();
-			UndoObject->SaveState(*InCurveEditor, GetDeformer());
-		}
-		
-		const TArray<FVector2D> TranslatedControlPoints = PlaceLatticeGridOverNewKeyPositions(GetDeformer(), *InCurveEditor);
-		checkSlow(TranslatedControlPoints.Num() == GetDeformer().GetControlPoints().Num());
-		for (int32 Index = 0; Index < TranslatedControlPoints.Num(); ++Index)
-		{
-			const FVector2D& UpdatedControlPoint = TranslatedControlPoints[Index];
-			const TConstArrayView<FVector2D> DummyView(&UpdatedControlPoint, 1);
-			GetDeformer().SetControlsPointNoRecompute({ Index }, DummyView);
-		}
-	}
+	// This will calculate a new rectangular grid around the points.
+	// In the past, we tried to preserve the existing shape by just moving it over the moved keys but there were too many different tools interacting
+	// which made this difficult to get right in every case. Hence, we decided just to recompute the grid on move.
+	LatticeDeformerState.Reset();
+	UpdateDeformerBounds(InCurveEditor);
 }
 
 void FCurveEditorLatticeTool::CancelDrag()
diff --git a/Engine/Plugins/Editor/CurveEditorTools/Source/CurveEditorTools/Private/Tools/Lattice/CurveEditorLatticeTool.h b/Engine/Plugins/Editor/CurveEditorTools/Source/CurveEditorTools/Private/Tools/Lattice/CurveEditorLatticeTool.h
--- a/Engine/Plugins/Editor/CurveEditorTools/Source/CurveEditorTools/Private/Tools/Lattice/CurveEditorLatticeTool.h
+++ b/Engine/Plugins/Editor/CurveEditorTools/Source/CurveEditorTools/Private/Tools/Lattice/CurveEditorLatticeTool.h
@@ -137,7 +137,7 @@ private:
 	/** Updates the bounds after the user has externally moved the keys, e.g. MMB move keys, etc. */
 	void OnCurvesModified();
 	/** Call this if the keys have externally moved. Tries to fit the lattice grid over the keys or reconstructs a new rectangular one if not possible. */
-	void AdjustLatticeGridOrRecompute(const TSharedRef<FCurveEditor>& InCurveEditor);
+	void RecomputeLatticeGrid(const TSharedRef<FCurveEditor>& InCurveEditor);
 	
 	/** Stops the active drag operation. */
 	void CancelDrag();
diff --git a/Engine/Plugins/Editor/CurveEditorTools/Source/CurveEditorTools/Private/Tools/Lattice/Misc/LatticeUtils.cpp b/Engine/Plugins/Editor/CurveEditorTools/Source/CurveEditorTools/Private/Tools/Lattice/Misc/LatticeUtils.cpp
--- a/Engine/Plugins/Editor/CurveEditorTools/Source/CurveEditorTools/Private/Tools/Lattice/Misc/LatticeUtils.cpp
+++ b/Engine/Plugins/Editor/CurveEditorTools/Source/CurveEditorTools/Private/Tools/Lattice/Misc/LatticeUtils.cpp
@@ -86,60 +86,6 @@ FCellVertexIndices GetCellIndices(int32 InCellIndex, int32 InNumPointsInWidth)
 	return CellIndices;
 }
 
-template<typename TCallback> requires std::is_invocable_v<TCallback, const FKeyPosition&>
-static void ForEachKey(const FCurveEditor& InCurveEditor, TCallback&& Callback)
-{
-	const TMap<FCurveModelID, FKeyHandleSet>& Selection = InCurveEditor.Selection.GetAll();
-	for (const TPair<FCurveModelID, FKeyHandleSet>& Pair : Selection)
-	{
-		FCurveModel* CurveModel = InCurveEditor.FindCurve(Pair.Key);
-		if (!CurveModel)
-		{
-			return;
-		}
-
-		const TConstArrayView<FKeyHandle> KeyHandles = Pair.Value.AsArray();
-		TArray<FKeyPosition> KeyPositions;
-		KeyPositions.SetNumUninitialized(KeyHandles.Num());
-		CurveModel->GetKeyPositions(KeyHandles, KeyPositions);
-		
-		for (const FKeyPosition& Position : KeyPositions)
-		{
-			Callback(Position);
-		}
-	}
-}
-
-TArray<FVector2D> PlaceLatticeGridOverNewKeyPositions(const FLatticeDeformer2D& InDeformer, const FCurveEditor& InCurveEditor)
-{
-	checkfSlow(!InCurveEditor.Selection.GetAll().IsEmpty(), TEXT("The key selection is supposed to be the same as when the deformer was created!")); 
-	FVector2D NewMinValues { TNumericLimits<double>::Max(), TNumericLimits<double>::Max() };;
-	ForEachKey(InCurveEditor, [&NewMinValues](const FKeyPosition& Position)
-	{
-		NewMinValues.X = FMath::Min(Position.InputValue, NewMinValues.X);
-		NewMinValues.Y = FMath::Min(Position.OutputValue, NewMinValues.Y);
-	});
-	
-	FVector2D OldMinValues { TNumericLimits<double>::Max(), TNumericLimits<double>::Max() };
-	for (const FVector2D& ControlPoint : InDeformer.GetControlPoints())
-	{
-		OldMinValues.X = FMath::Min(ControlPoint.X, OldMinValues.X);
-		OldMinValues.Y = FMath::Min(ControlPoint.Y, OldMinValues.Y);
-	}
-
-	// The algorithm assumes that all points have been translated uniformly, like you would if you move the key selection with middle mouse button.
-	const FVector2D OldMinToNewMin = NewMinValues - OldMinValues;
-	
-	TArray<FVector2D> TransformedControlPoints;
-	TransformedControlPoints.SetNumUninitialized(InDeformer.NumControlPoints());
-	for (int32 ControlPointIndex = 0; ControlPointIndex < InDeformer.NumControlPoints(); ++ControlPointIndex)
-	{
-		TransformedControlPoints[ControlPointIndex] = InDeformer.GetControlPoints()[ControlPointIndex] + OldMinToNewMin;
-	}
-	
-	return TransformedControlPoints;
-}
-
 void AddSelectionToDeformer(FLatticeDeformer2D& InDeformer, const FCurveEditor& InCurveEditor)
 {
 	for (const TPair<FCurveModelID, FKeyHandleSet>& Pair : InCurveEditor.Selection.GetAll())
diff --git a/Engine/Plugins/Editor/CurveEditorTools/Source/CurveEditorTools/Private/Tools/Lattice/Misc/LatticeUtils.h b/Engine/Plugins/Editor/CurveEditorTools/Source/CurveEditorTools/Private/Tools/Lattice/Misc/LatticeUtils.h
--- a/Engine/Plugins/Editor/CurveEditorTools/Source/CurveEditorTools/Private/Tools/Lattice/Misc/LatticeUtils.h
+++ b/Engine/Plugins/Editor/CurveEditorTools/Source/CurveEditorTools/Private/Tools/Lattice/Misc/LatticeUtils.h
@@ -68,13 +68,6 @@ struct FCellVertexIndices
  */
 FCellVertexIndices GetCellIndices(int32 InCellIndex, int32 InNumPointsInWidth);
 
-/**
- * Call this in response to the positions of the key selection having been moved.
- * Assuming that all points have been translated uniformly, this moves the the lattice grid by the delta movement.
- * @return The new control points that InDeformer should have, or empty, if the lattice grid needs to be regenerated.
- */
-TArray<FVector2D> PlaceLatticeGridOverNewKeyPositions(const FLatticeDeformer2D& InDeformer, const FCurveEditor& InCurveEditor);
-
 /** Adds InCurveEditor's selection to InDeformer. */
 void AddSelectionToDeformer(FLatticeDeformer2D& InDeformer, const FCurveEditor& InCurveEditor);
 
