Curve Editor: Lattice tool will now call SetKeyPositions with EPropertyChangeType::Interactive while dragging and EPropertyChangeType::ValueSet when the user releases the mouse.

diff --git a/Engine/Plugins/Editor/CurveEditorTools/Source/CurveEditorTools/Private/Tools/Lattice/CurveEditorLatticeTool.cpp b/Engine/Plugins/Editor/CurveEditorTools/Source/CurveEditorTools/Private/Tools/Lattice/CurveEditorLatticeTool.cpp
--- a/Engine/Plugins/Editor/CurveEditorTools/Source/CurveEditorTools/Private/Tools/Lattice/CurveEditorLatticeTool.cpp
+++ b/Engine/Plugins/Editor/CurveEditorTools/Source/CurveEditorTools/Private/Tools/Lattice/CurveEditorLatticeTool.cpp
@@ -388,9 +388,11 @@ void FCurveEditorLatticeTool::StartDragDrop(const FLatticeBounds& InBounds, cons
 			InBounds,
 			FLatticeDragOp_MoveControlPoints::FOnControlPointsMoved::CreateLambda(
 				[this, ControlPointIndex](const TConstArrayView<FVector2D>& InNewControlPointPositions)
-				{
-					OnMoveControlPoint({ ControlPointIndex }, InNewControlPointPositions);
-				}
+				{ OnMoveControlPoint({ ControlPointIndex }, InNewControlPointPositions, EPropertyChangeType::Interactive); }
+			),
+			FLatticeDragOp_MoveControlPoints::FOnControlPointsMoved::CreateLambda(
+				[this, ControlPointIndex](const TConstArrayView<FVector2D>& InNewControlPointPositions)
+				{ OnMoveControlPoint({ ControlPointIndex }, InNewControlPointPositions, EPropertyChangeType::ValueSet); }
 			)
 		);
 		DelayedDragData->ActiveOperation.Emplace(MoveTemp(Op), LOCTEXT("MoveControlPoint", "Move control point"));
@@ -409,9 +411,11 @@ void FCurveEditorLatticeTool::StartDragDrop(const FLatticeBounds& InBounds, cons
 			InBounds,
 			FLatticeDragOp_MoveControlPoints::FOnControlPointsMoved::CreateLambda(
 				[this, Start, End](const TConstArrayView<FVector2D>& InNewControlPointPositions)
-				{
-					OnMoveEdge({ Start, End }, InNewControlPointPositions);
-				}
+				{ OnMoveEdge({ Start, End }, InNewControlPointPositions, EPropertyChangeType::Interactive); }
+			),
+			FLatticeDragOp_MoveControlPoints::FOnControlPointsMoved::CreateLambda(
+				[this, Start, End](const TConstArrayView<FVector2D>& InNewControlPointPositions)
+				{ OnMoveEdge({ Start, End }, InNewControlPointPositions, EPropertyChangeType::ValueSet); }
 			)
 		);
 		DelayedDragData->ActiveOperation.Emplace(MoveTemp(Op), LOCTEXT("MoveEdge", "Move edge"));
@@ -430,7 +434,13 @@ void FCurveEditorLatticeTool::StartDragDrop(const FLatticeBounds& InBounds, cons
 			FLatticeDragOp_MoveControlPoints::FOnControlPointsMoved::CreateLambda(
 				[this, TopLeft, TopRight, BottomRight, BottomLeft](const TConstArrayView<FVector2D>& InNewControlPointPositions)
 				{
-					OnSetControlPoints({ TopLeft, TopRight, BottomRight, BottomLeft }, InNewControlPointPositions);
+					OnSetControlPoints({ TopLeft, TopRight, BottomRight, BottomLeft }, InNewControlPointPositions, EPropertyChangeType::Interactive);
+				}
+			),
+			FLatticeDragOp_MoveControlPoints::FOnControlPointsMoved::CreateLambda(
+				[this, TopLeft, TopRight, BottomRight, BottomLeft](const TConstArrayView<FVector2D>& InNewControlPointPositions)
+				{
+					OnSetControlPoints({ TopLeft, TopRight, BottomRight, BottomLeft }, InNewControlPointPositions, EPropertyChangeType::ValueSet);
 				}
 			)
 		);
@@ -438,7 +448,7 @@ void FCurveEditorLatticeTool::StartDragDrop(const FLatticeBounds& InBounds, cons
 	}
 }
 
-void FCurveEditorLatticeTool::OnMoveControlPoint(TConstArrayView<int32> InVertexIndices, TConstArrayView<FVector2D> InNewControlPointPositions)
+void FCurveEditorLatticeTool::OnMoveControlPoint(TConstArrayView<int32> InVertexIndices, TConstArrayView<FVector2D> InNewControlPointPositions, EPropertyChangeType::Type InChangeType)
 {
 	check(InVertexIndices.Num() == 1 && InNewControlPointPositions.Num() == 1);
 	
@@ -447,10 +457,10 @@ void FCurveEditorLatticeTool::OnMoveControlPoint(TConstArrayView<int32> InVertex
 		DelayedDragData->PointDragTangentMirrorOperation->OnMovePoint(InNewControlPointPositions[0]);
 	}
 	
-	OnSetControlPoints(InVertexIndices, InNewControlPointPositions);
+	OnSetControlPoints(InVertexIndices, InNewControlPointPositions, InChangeType);
 }
 
-void FCurveEditorLatticeTool::OnMoveEdge(TConstArrayView<int32> InVertexIndices, TConstArrayView<FVector2D> InNewControlPointPositions)
+void FCurveEditorLatticeTool::OnMoveEdge(TConstArrayView<int32> InVertexIndices, TConstArrayView<FVector2D> InNewControlPointPositions, EPropertyChangeType::Type InChangeType)
 {
 	check(InVertexIndices.Num() == 2 && InNewControlPointPositions.Num() == 2);
 	
@@ -459,10 +469,12 @@ void FCurveEditorLatticeTool::OnMoveEdge(TConstArrayView<int32> InVertexIndices,
 		DelayedDragData->EdgeDragTangentMirrorOperation->OnMoveEdge(InNewControlPointPositions);
 	}
 
-	OnSetControlPoints(InVertexIndices, InNewControlPointPositions);
+	OnSetControlPoints(InVertexIndices, InNewControlPointPositions, InChangeType);
 }
 
-void FCurveEditorLatticeTool::OnSetControlPoints(TConstArrayView<int32> InControlPointsToUpdate, TConstArrayView<FVector2D> InNewControlPoints)
+void FCurveEditorLatticeTool::OnSetControlPoints(
+	TConstArrayView<int32> InControlPointsToUpdate, TConstArrayView<FVector2D> InNewControlPoints, EPropertyChangeType::Type InChangeType
+	)
 {
 	const TSharedPtr<FCurveEditor> CurveEditor = WeakCurveEditor.Pin();
 	const TOptional<FCurvePointSnapper> Snapper = CurveEditor ? FCurvePointSnapper::MakeSnapper(*CurveEditor) : TOptional<FCurvePointSnapper>{};
@@ -516,7 +528,7 @@ void FCurveEditorLatticeTool::OnSetControlPoints(TConstArrayView<int32> InContro
 		if (FCurveModel* Model = CurveEditor->FindCurve(It->Key); ensure(Model))
 		{
 			Model->Modify();
-			Model->SetKeyPositions(Pair.Handles, Pair.Keys);
+			Model->SetKeyPositions(Pair.Handles, Pair.Keys, InChangeType);
 		}
 	}
 }
@@ -556,7 +568,7 @@ void FCurveEditorLatticeTool::MoveEdgeToOpposite(int32 InEdgeIndex)
 		MirrorOp->OnMoveEdge(NewEdge);
 	}
 	// ... then move the grid
-	OnSetControlPoints({ MovedEdgeStartVertIdx, MovedEdgeEndVertIdx }, NewEdge);
+	OnSetControlPoints({ MovedEdgeStartVertIdx, MovedEdgeEndVertIdx }, NewEdge, EPropertyChangeType::ValueSet);
 }
 
 void FCurveEditorLatticeTool::MirrorGrid()
diff --git a/Engine/Plugins/Editor/CurveEditorTools/Source/CurveEditorTools/Private/Tools/Lattice/CurveEditorLatticeTool.h b/Engine/Plugins/Editor/CurveEditorTools/Source/CurveEditorTools/Private/Tools/Lattice/CurveEditorLatticeTool.h
--- a/Engine/Plugins/Editor/CurveEditorTools/Source/CurveEditorTools/Private/Tools/Lattice/CurveEditorLatticeTool.h
+++ b/Engine/Plugins/Editor/CurveEditorTools/Source/CurveEditorTools/Private/Tools/Lattice/CurveEditorLatticeTool.h
@@ -13,6 +13,7 @@
 #include "Misc/LatticeDrawUtils.h"
 #include "ScopedTransaction.h"
 #include "Templates/SharedPointer.h"
+#include "UObject/UnrealType.h"
 
 class FCurveEditor;
 class FText;
@@ -151,11 +152,11 @@ private:
 	/** Starts a new drag drop operation based on whether a control point, edge, or cell is hovered. */
 	void StartDragDrop(const FLatticeBounds& InBounds, const FLatticeHoverState& InHoverState);
 	/** Called when a control point is dragged. */
-	void OnMoveControlPoint(TConstArrayView<int32> InVertexIndices, TConstArrayView<FVector2D> InNewControlPointPositions);
+	void OnMoveControlPoint(TConstArrayView<int32> InVertexIndices, TConstArrayView<FVector2D> InNewControlPointPositions, EPropertyChangeType::Type InChangeType);
 	/** Called when an edge is dragged. */
-	void OnMoveEdge(TConstArrayView<int32> InVertexIndices, TConstArrayView<FVector2D> InNewControlPointPositions);
+	void OnMoveEdge(TConstArrayView<int32> InVertexIndices, TConstArrayView<FVector2D> InNewControlPointPositions, EPropertyChangeType::Type InChangeType);
 	/** Generic function for updating control points. */
-	void OnSetControlPoints(TConstArrayView<int32> InControlPointsToUpdate, TConstArrayView<FVector2D> InNewControlPoints);
+	void OnSetControlPoints(TConstArrayView<int32> InControlPointsToUpdate, TConstArrayView<FVector2D> InNewControlPoints, EPropertyChangeType::Type InChangeType);
 
 	/** Moves the clicked edge to the opposite edge, effectively flattening the curve sections. */
 	void MoveEdgeToOpposite(int32 InEdgeIndex);
diff --git a/Engine/Plugins/Editor/CurveEditorTools/Source/CurveEditorTools/Private/Tools/Lattice/Drag/LatticeDragOp_MoveControlPoints.cpp b/Engine/Plugins/Editor/CurveEditorTools/Source/CurveEditorTools/Private/Tools/Lattice/Drag/LatticeDragOp_MoveControlPoints.cpp
--- a/Engine/Plugins/Editor/CurveEditorTools/Source/CurveEditorTools/Private/Tools/Lattice/Drag/LatticeDragOp_MoveControlPoints.cpp
+++ b/Engine/Plugins/Editor/CurveEditorTools/Source/CurveEditorTools/Private/Tools/Lattice/Drag/LatticeDragOp_MoveControlPoints.cpp
@@ -11,16 +11,30 @@ FLatticeDragOp_MoveControlPoints::FLatticeDragOp_MoveControlPoints(
 	TWeakPtr<FCurveEditor> InCurveEditor,
 	FControlPointArray InInitialControlPoints,
 	const FLatticeBounds& InBounds,
-	FOnControlPointsMoved SetControlPointDelegate
+	FOnControlPointsMoved InInteractiveDragControlPointsDelegate,
+	FOnControlPointsMoved InFinishDragControlPointsDelegate
 	)
 	: FLatticeDragOp(MoveTemp(InCurveEditor))
 	, InitialControlPoints(MoveTemp(InInitialControlPoints))
 	, Bounds(InBounds)
-	, SetControlPointDelegate(MoveTemp(SetControlPointDelegate))
+	, OnInteractiveDragControlPointsDelegate(MoveTemp(InInteractiveDragControlPointsDelegate))
+	, OnFinishDragControlPointsDelegate(MoveTemp(InFinishDragControlPointsDelegate))
 	, UpdatedControlPointPositions(InitialControlPoints)
 {}
 	
 void FLatticeDragOp_MoveControlPoints::OnMoveMouse(const FGeometry& InGeometry, const FVector2D& InScreenPosition)
+{
+	UpdateControlPoints(InGeometry, InScreenPosition);
+	OnInteractiveDragControlPointsDelegate.Execute(UpdatedControlPointPositions);
+}
+
+void FLatticeDragOp_MoveControlPoints::OnEndDrag(const FGeometry& InGeometry, const FVector2D& InMousePosition)
+{
+	UpdateControlPoints(InGeometry, InMousePosition);
+	OnFinishDragControlPointsDelegate.Execute(UpdatedControlPointPositions);
+}
+
+void FLatticeDragOp_MoveControlPoints::UpdateControlPoints(const FGeometry& InGeometry, const FVector2D& InScreenPosition)
 {
 	const FSlateLayoutTransform ScreenToWidget = InGeometry.GetAccumulatedLayoutTransform().Inverse();
 	const FVector2D Initial_WidgetSpace = ScreenToWidget.TransformPoint(GetInitialMousePosition());
@@ -35,6 +49,5 @@ void FLatticeDragOp_MoveControlPoints::OnMoveMouse(const FGeometry& InGeometry,
 	{
 		UpdatedControlPointPositions[Index] = InitialControlPoints[Index] + DeltaMove_ControlPointSpace;
 	}
-	SetControlPointDelegate.Execute(UpdatedControlPointPositions);
 }
 }
diff --git a/Engine/Plugins/Editor/CurveEditorTools/Source/CurveEditorTools/Private/Tools/Lattice/Drag/LatticeDragOp_MoveControlPoints.h b/Engine/Plugins/Editor/CurveEditorTools/Source/CurveEditorTools/Private/Tools/Lattice/Drag/LatticeDragOp_MoveControlPoints.h
--- a/Engine/Plugins/Editor/CurveEditorTools/Source/CurveEditorTools/Private/Tools/Lattice/Drag/LatticeDragOp_MoveControlPoints.h
+++ b/Engine/Plugins/Editor/CurveEditorTools/Source/CurveEditorTools/Private/Tools/Lattice/Drag/LatticeDragOp_MoveControlPoints.h
@@ -27,13 +27,15 @@ public:
 		TWeakPtr<FCurveEditor> InCurveEditor,
 		FControlPointArray InInitialControlPoints,
 		const FLatticeBounds& InBounds,
-		FOnControlPointsMoved SetControlPointDelegate
+		FOnControlPointsMoved InInteractiveDragControlPointsDelegate,
+		FOnControlPointsMoved InFinishDragControlPointsDelegate
 		);
 
 protected:
 	
 	//~ Begin FLatticeDragOp Interface
 	virtual void OnMoveMouse(const FGeometry& InGeometry, const FVector2D& InScreenPosition) override;
+	virtual void OnEndDrag(const FGeometry& InGeometry, const FVector2D& InMousePosition) override;
 	//~ End FLatticeDragOp Interface
 
 private:
@@ -42,11 +44,17 @@ private:
 	const FControlPointArray InitialControlPoints;
 	/** Used to transform points */
 	const FLatticeBounds Bounds;
-	/** Invoked to change the control points by the. */
-	const FOnControlPointsMoved SetControlPointDelegate;
+	
+	/** Invoked while the mouse button is down. */
+	const FOnControlPointsMoved OnInteractiveDragControlPointsDelegate;
+	/** Invoked one last time when the mouse button is lifted. */
+	const FOnControlPointsMoved OnFinishDragControlPointsDelegate;
 
 	/** Updated every time the mouse moves: InitialControlPoints + DeltaMouseMove. */
 	FControlPointArray UpdatedControlPointPositions;
+
+	/** Updates UpdatedControlPointPositions. */
+	void UpdateControlPoints(const FGeometry& InGeometry, const FVector2D& InScreenPosition);
 };
 }
 
