Curve Editor: Fix lattice jump on mouse release when axis snapping enabled

diff --git a/Engine/Plugins/Editor/CurveEditorTools/Source/CurveEditorTools/Private/Tools/Lattice/Drag/LatticeDragOp.cpp b/Engine/Plugins/Editor/CurveEditorTools/Source/CurveEditorTools/Private/Tools/Lattice/Drag/LatticeDragOp.cpp
--- a/Engine/Plugins/Editor/CurveEditorTools/Source/CurveEditorTools/Private/Tools/Lattice/Drag/LatticeDragOp.cpp
+++ b/Engine/Plugins/Editor/CurveEditorTools/Source/CurveEditorTools/Private/Tools/Lattice/Drag/LatticeDragOp.cpp
@@ -34,7 +34,9 @@ void FLatticeDragOp::FinishedPointerInput()
 
 void FLatticeDragOp::EndDrag(const FGeometry& InGeometry, const FPointerEvent& InMouseEvent)
 {
-	OnEndDrag(InGeometry, InMouseEvent.GetScreenSpacePosition());
+	// Snap the mouse position according to the axis snapping settings.
+	AccumulateMouseMovement(InGeometry, InMouseEvent);
+	OnEndDrag(InGeometry, AccumulatedMouseMovement->AccumulatedPosition);
 	AccumulatedMouseMovement.Reset();
 }
 
