Curve Editor Tools: Fix failing checkSlow. Lattice tool now only registers for undo / redo events while the tool is active.

diff --git a/Engine/Plugins/Editor/CurveEditorTools/Source/CurveEditorTools/Private/Tools/Lattice/CurveEditorLatticeTool.cpp b/Engine/Plugins/Editor/CurveEditorTools/Source/CurveEditorTools/Private/Tools/Lattice/CurveEditorLatticeTool.cpp
--- a/Engine/Plugins/Editor/CurveEditorTools/Source/CurveEditorTools/Private/Tools/Lattice/CurveEditorLatticeTool.cpp
+++ b/Engine/Plugins/Editor/CurveEditorTools/Source/CurveEditorTools/Private/Tools/Lattice/CurveEditorLatticeTool.cpp
@@ -57,19 +57,11 @@ FCurveEditorLatticeTool::FCurveEditorLatticeTool(TWeakPtr<FCurveEditor> InCurveE
 			UndoObject->RestoreEditorState(*Editor);
 		}
 	});
-	
-	if (UTransBuffer* Transactor = GEditor ? Cast<UTransBuffer>(GEditor->Trans) : nullptr)
-	{
-		Transactor->OnTransactionStateChanged().AddRaw(this, &FCurveEditorLatticeTool::OnTransactionStateChanged);
-	}
 }
 
 FCurveEditorLatticeTool::~FCurveEditorLatticeTool()
 {
-	if (UTransBuffer* Transactor = GEditor ? Cast<UTransBuffer>(GEditor->Trans) : nullptr)
-	{
-		Transactor->OnTransactionStateChanged().RemoveAll(this);
-	}
+	UnregisterTransactionStateChanged();
 }
 
 void FCurveEditorLatticeTool::Tick(const FGeometry& InViewContainerGeometry, const double InCurrentTime, const float InDeltaTime)
@@ -240,12 +232,18 @@ void FCurveEditorLatticeTool::OnToolActivated()
 	{
 		UpdateDeformerBounds(CurveEditor.ToSharedRef());
 	}
+	
+	if (UTransBuffer* Transactor = GEditor ? Cast<UTransBuffer>(GEditor->Trans) : nullptr)
+	{
+		Transactor->OnTransactionStateChanged().AddRaw(this, &FCurveEditorLatticeTool::OnTransactionStateChanged);
+	}
 }
 
 void FCurveEditorLatticeTool::OnToolDeactivated()
 {
 	// Clear data. We'll recompute a new grid next time the user activate the tool
 	LatticeDeformerState.Reset();
+	UnregisterTransactionStateChanged();
 }
 
 void FCurveEditorLatticeTool::BindCommands(TSharedRef<FUICommandList> CommandBindings)
@@ -340,7 +338,7 @@ void FCurveEditorLatticeTool::OnCurvesModified()
 }
 
 void FCurveEditorLatticeTool::OnTransactionStateChanged(
-	const FTransactionContext&, ETransactionStateEventType TransactionStateEvent
+	const FTransactionContext& Context, ETransactionStateEventType TransactionStateEvent
 	)
 {
 	// Technically, we should check whether the transaction actually affected the lattice... but this isn't computationally expensive so whatever.
@@ -353,8 +351,6 @@ void FCurveEditorLatticeTool::OnTransactionStateChanged(
 
 void FCurveEditorLatticeTool::RecomputeLatticeGrid(const TSharedRef<FCurveEditor>& InCurveEditor, ELatticeUpdateFlags InRefreshFlags)
 {
-	checkSlow(LatticeDeformerState);
-	
 	// This will calculate a new rectangular grid around the points.
 	// In the past, we tried to preserve the existing shape by just moving it over the moved keys but there were too many different tools interacting
 	// which made this difficult to get right in every case. Hence, we decided just to recompute the grid on move.
@@ -672,6 +668,14 @@ bool FCurveEditorLatticeTool::ComputeHoverStateForLeftClick(
 
 	return false;
 }
+	
+void FCurveEditorLatticeTool::UnregisterTransactionStateChanged() const
+{
+	if (UTransBuffer* Transactor = GEditor ? Cast<UTransBuffer>(GEditor->Trans) : nullptr)
+	{
+		Transactor->OnTransactionStateChanged().RemoveAll(this);
+	}
+}
 }
 
 #undef LOCTEXT_NAMESPACE
\ No newline at end of file
diff --git a/Engine/Plugins/Editor/CurveEditorTools/Source/CurveEditorTools/Private/Tools/Lattice/CurveEditorLatticeTool.h b/Engine/Plugins/Editor/CurveEditorTools/Source/CurveEditorTools/Private/Tools/Lattice/CurveEditorLatticeTool.h
--- a/Engine/Plugins/Editor/CurveEditorTools/Source/CurveEditorTools/Private/Tools/Lattice/CurveEditorLatticeTool.h
+++ b/Engine/Plugins/Editor/CurveEditorTools/Source/CurveEditorTools/Private/Tools/Lattice/CurveEditorLatticeTool.h
@@ -189,6 +189,9 @@ private:
 	 * @return False if InCallback was not invoked. Otherwise, whatever InCallback returned. */
 	template<typename TCallback> requires std::is_invocable_r_v<bool, TCallback, const FLatticeHoverState&, const FLatticeBounds&>
 	bool ComputeHoverStateForLeftClick(const FGeometry& InViewContainerGeometry, const FPointerEvent& InMouseEvent, TCallback&& InCallback);
+
+	/** Shared logic for unregistering global transaction delegate. */
+	void UnregisterTransactionStateChanged() const;
 };
 
 inline void FCurveEditorLatticeTool::RecomputeLatticeGrid(ELatticeUpdateFlags InRefreshFlags)
