Curve Editor: Fix slow down when double clicking lattice top or bottom edge due to ensure. The ensure triggers when all keys in the lattice selection are set to auto interpolate.

diff --git a/Engine/Plugins/Editor/CurveEditorTools/Source/CurveEditorTools/Private/Tools/Lattice/CurveEditorLatticeTool.cpp b/Engine/Plugins/Editor/CurveEditorTools/Source/CurveEditorTools/Private/Tools/Lattice/CurveEditorLatticeTool.cpp
--- a/Engine/Plugins/Editor/CurveEditorTools/Source/CurveEditorTools/Private/Tools/Lattice/CurveEditorLatticeTool.cpp
+++ b/Engine/Plugins/Editor/CurveEditorTools/Source/CurveEditorTools/Private/Tools/Lattice/CurveEditorLatticeTool.cpp
@@ -607,7 +607,8 @@ void FCurveEditorLatticeTool::MoveEdgeToOpposite(int32 InEdgeIndex)
 	TOptional<FLatticeEdgeTangentsMirrorOp> MirrorOp = FLatticeEdgeTangentsMirrorOp::MakeFromDraggedEdge(
 		OppositeEdgeIndex, GetDeformer(), *WeakCurveEditor.Pin()
 		);
-	if (ensure(MirrorOp)) // Should be set because we've checked above that we're moving a top or bottom edge
+	// Even though we've checked that we're moving a top or bottom edge, the mirror op can still be unset if all keys are set to auto interpolate.
+	if (MirrorOp) 
 	{
 		MirrorOp->OnMoveEdge(NewEdge);
 	}
