Curve Editor: Lattice tool now works in all views. Normalized view was the target, and Stacked view works as a bonus (with edge cases feeling odd).

Transformation overview:
- There is now a UI-only deformer ("global deformer") and per-curve deformers
- Each per-curve has a linear transformation that transforms from the global deformer to the per-curve transformer
- In Absolute mode, the transformation is the identify matrix
- In Normalized mode, the transformation transforms the lattice points such that the lattice is in normalized coordinate
- See FPerCurveLatticeData::GlobalDeformerToCurveDeformer for details

Undo:
- Lattice no longer restores selection
- Fixed bug where selecting keys would generate 2 transactions (one from core CurveEditor module, and one from lattice)
- Undo continues to restore the control point positions, however, only as long as you stay in the same view.
- If you switched views, the lattice grid is regenerated to be rectuangular.

diff --git a/Engine/Plugins/Editor/CurveEditorTools/Source/CurveEditorTools/Private/Misc/CurveViewRebuildListener.cpp b/Engine/Plugins/Editor/CurveEditorTools/Source/CurveEditorTools/Private/Misc/CurveViewRebuildListener.cpp
new file mode 100644
--- /dev/null
+++ b/Engine/Plugins/Editor/CurveEditorTools/Source/CurveEditorTools/Private/Misc/CurveViewRebuildListener.cpp
@@ -0,0 +1,29 @@
+﻿// Copyright Epic Games, Inc. All Rights Reserved.
+
+#include "CurveViewRebuildListener.h"
+
+#include "CurveEditor.h"
+#include "SCurveEditorPanel.h"
+
+namespace UE::CurveEditorTools
+{
+FCurveViewRebuildListener::FCurveViewRebuildListener(const TSharedRef<FCurveEditor>& InCurveEditor)
+	: WeakCurveEditor(InCurveEditor)
+{
+	const TSharedPtr<SCurveEditorPanel> Panel = InCurveEditor->GetPanel();
+	if (ensure(Panel))
+	{
+		Panel->OnPostRebuildCurveViews().AddRaw(this, &FCurveViewRebuildListener::BroadcastOnCurveModified);
+	}
+}
+
+FCurveViewRebuildListener::~FCurveViewRebuildListener()
+{
+	const TSharedPtr<FCurveEditor> CurveEditorPin = WeakCurveEditor.Pin();
+	const TSharedPtr<SCurveEditorPanel> Panel = CurveEditorPin ? CurveEditorPin->GetPanel() : nullptr;
+	if (Panel)
+	{
+		Panel->OnPostRebuildCurveViews().RemoveAll(this);
+	}
+}
+}
diff --git a/Engine/Plugins/Editor/CurveEditorTools/Source/CurveEditorTools/Private/Misc/CurveViewRebuildListener.h b/Engine/Plugins/Editor/CurveEditorTools/Source/CurveEditorTools/Private/Misc/CurveViewRebuildListener.h
new file mode 100644
--- /dev/null
+++ b/Engine/Plugins/Editor/CurveEditorTools/Source/CurveEditorTools/Private/Misc/CurveViewRebuildListener.h
@@ -0,0 +1,37 @@
+﻿// Copyright Epic Games, Inc. All Rights Reserved.
+
+#pragma once
+
+#include "Delegates/Delegate.h"
+#include "Templates/SharedPointer.h"
+#include "Templates/UnrealTemplate.h"
+
+class FCurveEditor;
+struct FCurveModelID;
+
+namespace UE::CurveEditorTools
+{
+/**
+ * Invokes OnCurveViewRebuilt when the curve editor's SCurveEditorPanel::OnPostRebuildCurveViews is invoked.
+ * This handles clean-up, etc. so the using code does not need to worry about it.
+ */
+class FCurveViewRebuildListener : public FNoncopyable
+{
+public:
+
+	explicit FCurveViewRebuildListener(const TSharedRef<FCurveEditor>& InCurveEditor);
+	~FCurveViewRebuildListener();
+
+	FSimpleMulticastDelegate& OnCurveViewRebuilt() { return OnCurveModifiedDelegate; }
+
+private:
+
+	/** Used to unsubscribe on destruction. */
+	TWeakPtr<FCurveEditor> WeakCurveEditor;
+
+	/** Invokes when any of the listeners is changed. */
+	FSimpleMulticastDelegate OnCurveModifiedDelegate;
+
+	void BroadcastOnCurveModified() const { OnCurveModifiedDelegate.Broadcast(); }
+};
+}
\ No newline at end of file
diff --git a/Engine/Plugins/Editor/CurveEditorTools/Source/CurveEditorTools/Private/Misc/KeyInCurveId.cpp b/Engine/Plugins/Editor/CurveEditorTools/Source/CurveEditorTools/Private/Misc/KeyInCurveId.cpp
deleted file mode 100644
diff --git a/Engine/Plugins/Editor/CurveEditorTools/Source/CurveEditorTools/Private/Misc/KeyInCurveId.h b/Engine/Plugins/Editor/CurveEditorTools/Source/CurveEditorTools/Private/Misc/KeyInCurveId.h
deleted file mode 100644
diff --git a/Engine/Plugins/Editor/CurveEditorTools/Source/CurveEditorTools/Private/Misc/VectorMathUtils.cpp b/Engine/Plugins/Editor/CurveEditorTools/Source/CurveEditorTools/Private/Misc/VectorMathUtils.cpp
--- a/Engine/Plugins/Editor/CurveEditorTools/Source/CurveEditorTools/Private/Misc/VectorMathUtils.cpp
+++ b/Engine/Plugins/Editor/CurveEditorTools/Source/CurveEditorTools/Private/Misc/VectorMathUtils.cpp
@@ -48,4 +48,38 @@ void InsetQuadBy(FVector2D& A, FVector2D& B, FVector2D& C, FVector2D& D, float I
 	C = C + -InsetAmount * (InwardBC + InwardCD).GetSafeNormal();
 	D = D + -InsetAmount * (InwardCD + InwardDA).GetSafeNormal();
 }
+
+FTransform2d TransformRectBetweenSpaces(
+	const FVector2D& InMinSource, const FVector2D& InMaxSource,
+	const FVector2D& InMinTarget, const FVector2D& InMaxTarget
+	)
+{
+	const FVector2D DeltaLattice = InMaxSource - InMinSource;
+	const FVector2D DeltaCurveSpace = InMaxTarget - InMinTarget; 
+	const FTransform2d ToCurveSpace = Concatenate(
+		-InMinSource,									// Translate absolute to origin 
+		FScale2d(FVector2D(1., 1.) / DeltaLattice),	// Normalize to [0,1] range
+		FScale2d(DeltaCurveSpace),					// Rescale to curve space
+		InMinTarget								// Translate to curve space offset
+		);
+	return ToCurveSpace;
+}
+	
+FVector2D TransformAbsoluteToCurveSpace(const FTransform2d& InAbsToCurveSpace, const FVector2D& InPoint)
+{
+	double M00, M01, M10, M11;
+	InAbsToCurveSpace.GetMatrix().GetMatrix(M00, M01, M10, M11);
+
+	FVector2D CurveSpacePosition = InPoint - InAbsToCurveSpace.GetTranslation();
+	CurveSpacePosition.X /= M00;
+	CurveSpacePosition.Y /= M11;
+	return CurveSpacePosition;
+}
+	
+FVector2D TransformCurveSpaceToAbsolute(const FTransform2d& InAbsToCurveSpace, const FVector2D& InPoint)
+{
+	double M00, M01, M10, M11;
+	InAbsToCurveSpace.GetMatrix().GetMatrix(M00, M01, M10, M11);
+	return FVector2D(M00, M11) * InPoint + InAbsToCurveSpace.GetTranslation();
+}
 }
diff --git a/Engine/Plugins/Editor/CurveEditorTools/Source/CurveEditorTools/Private/Misc/VectorMathUtils.h b/Engine/Plugins/Editor/CurveEditorTools/Source/CurveEditorTools/Private/Misc/VectorMathUtils.h
--- a/Engine/Plugins/Editor/CurveEditorTools/Source/CurveEditorTools/Private/Misc/VectorMathUtils.h
+++ b/Engine/Plugins/Editor/CurveEditorTools/Source/CurveEditorTools/Private/Misc/VectorMathUtils.h
@@ -3,11 +3,48 @@
 #pragma once
 
 #include "Math/MathFwd.h"
+#include "Math/TransformCalculus2D.h"
 
 namespace UE::CurveEditorTools
 {
-	/** @return Whether Point is in the triangle spanned by the input vertices. */
-	bool IsPointInTriangle(const FVector2D& Point, const FVector2D& Vertex1, const FVector2D& Vertex2, const FVector2D& Vertex3);
-	/** @return Makes a smaller quad in the same shape but with the edges moved inside by InsetAmount. */
-	void InsetQuadBy(FVector2D& A, FVector2D& B, FVector2D& C, FVector2D& D, float InsetAmount);
+/** @return Whether Point is in the triangle spanned by the input vertices. */
+bool IsPointInTriangle(const FVector2D& Point, const FVector2D& Vertex1, const FVector2D& Vertex2, const FVector2D& Vertex3);
+/** @return Makes a smaller quad in the same shape but with the edges moved inside by InsetAmount. */
+void InsetQuadBy(FVector2D& A, FVector2D& B, FVector2D& C, FVector2D& D, float InsetAmount);
+
+/**
+ * Computes the transform for transforming a rectangle from source space to another target space.
+ *
+ * @param InMinSource The bottom left corner of the rectangle in absolute space
+ * @param InMaxSource The top right corner of the rectangle in absolute space
+ * @param InMinTarget The equivalent bottom left corner in curve space
+ * @param InMaxTarget The equivalent top right corner in curve space
+ */
+FTransform2d TransformRectBetweenSpaces(
+	const FVector2D& InMinSource, const FVector2D& InMaxSource,
+	const FVector2D& InMinTarget, const FVector2D& InMaxTarget
+	);
+
+/**
+ * Transforms a point from absolute key space to a SCurveEditorView's curve space.
+ * 
+ * Absolute key space are the literal FKeyPosition::InputValue and OutputValue.
+ * Curve space are the axis values the SCurveEditorView displays the keys at (e.g. in Normalized the max value is displayed at curve space value 1.0).
+ * 
+ * @param InAbsToCurveSpace Result of SCurveEditorView::GetViewToCurveTransform
+ * @param InPoint The point in absolute space to transform
+ * @return The point in curve space
+ */
+FVector2D TransformAbsoluteToCurveSpace(const FTransform2d& InAbsToCurveSpace, const FVector2D& InPoint);
+/**
+ * Transforms a point from SCurveEditorView's curve space to absolute key space.
+ * 
+ * Absolute key space are the literal FKeyPosition::InputValue and OutputValue.
+ * Curve space are the axis values the SCurveEditorView displays the keys at (e.g. in Normalized the max value is displayed at curve space value 1.0).
+ * 
+ * @param InAbsToCurveSpace Result of SCurveEditorView::GetViewToCurveTransform
+ * @param InPoint The point in curve space to transform
+ * @return The point in absolute space
+ */
+FVector2D TransformCurveSpaceToAbsolute(const FTransform2d& InAbsToCurveSpace, const FVector2D& InPoint);
 }
diff --git a/Engine/Plugins/Editor/CurveEditorTools/Source/CurveEditorTools/Private/Tools/Lattice/CurveEditorLatticeTool.cpp b/Engine/Plugins/Editor/CurveEditorTools/Source/CurveEditorTools/Private/Tools/Lattice/CurveEditorLatticeTool.cpp
--- a/Engine/Plugins/Editor/CurveEditorTools/Source/CurveEditorTools/Private/Tools/Lattice/CurveEditorLatticeTool.cpp
+++ b/Engine/Plugins/Editor/CurveEditorTools/Source/CurveEditorTools/Private/Tools/Lattice/CurveEditorLatticeTool.cpp
@@ -2,9 +2,9 @@
 
 #include "CurveEditorLatticeTool.h"
 
+#include "CoreGlobals.h"
 #include "CurveEditor.h"
 #include "CurveEditorToolCommands.h"
-#include "CoreGlobals.h"
 #include "Drag/LatticeDragOp_MoveControlPoints.h"
 #include "Editor.h"
 #include "Editor/TransBuffer.h"
@@ -12,14 +12,17 @@
 #include "Mirror/LatticeEdgeTangentMirrorOp.h"
 #include "Mirror/LatticePointTangentMirrorOp.h"
 #include "Misc/CurvePointSnapper.h"
+#include "Misc/DeformerUtils.h"
 #include "Misc/LatticeDrawUtils.h"
 #include "Misc/LatticeUtils.h"
-#include "Misc/VectorMathUtils.h"
 #include "Misc/Mirror/MirrorUtils.h"
+#include "Misc/VectorMathUtils.h"
+#include "SCurveEditorPanel.h"
+#include "SCurveEditorView.h"
 #include "Templates/UnrealTemplate.h"
-#include "UObject/LatticeUndoData.h"
 #include "UObject/Package.h"
 #include "UObject/UObjectGlobals.h"
+#include "Undo/LatticeUndoObject.h"
 
 #define LOCTEXT_NAMESPACE "FCurveEditorLatticeTool"
 
@@ -29,7 +32,7 @@ FCurveEditorLatticeTool::FCurveEditorLatticeTool(TWeakPtr<FCurveEditor> InCurveE
 	: WeakCurveEditor(InCurveEditor)
 	, FinalMouseScreenSpacePosition()
 	, UndoObject(
-		NewObject<UCurveEditorTools_LatticeUndoData>(GetTransientPackage(), UCurveEditorTools_LatticeUndoData::StaticClass(), NAME_None, RF_Transactional)
+		NewObject<UCurveEditorTools_LatticeUndoObject>(GetTransientPackage(), UCurveEditorTools_LatticeUndoObject::StaticClass(), NAME_None, RF_Transactional)
 	)
 {
 	UndoObject->OnPostEditUndo.AddLambda([this, CurveEditor = InCurveEditor]
@@ -43,25 +46,15 @@ FCurveEditorLatticeTool::FCurveEditorLatticeTool(TWeakPtr<FCurveEditor> InCurveE
 		// LatticeDeformerState is unset if the user redo-es the first transaction which set the initial selection...
 		if (LatticeDeformerState)
 		{
-			UndoObject->RestoreDeformerState(GetDeformer());
-			
-			UndoObject->RestoreEditorState(*Editor);
-			// We need to overwrite our serial number of next frame UpdateDeformerBounds will think that the selection has changed and reinit.
-			SelectionSerialNumber = UndoObject->GetSelectionSerialNumber();
-
+			ApplyUndoState(*Editor);
 			// If user undo-es the first transaction, next tick we'll notice that nothing is selected and destroy the deformer.
 		}
-		else
-		{
-			// ... in that case, we refresh the selection and next tick, we'll compute the bounds, which will end up have the rectangular starting shape.
-			UndoObject->RestoreEditorState(*Editor);
-		}
 	});
 }
 
 FCurveEditorLatticeTool::~FCurveEditorLatticeTool()
 {
-	UnregisterTransactionStateChanged();
+	UnregisterDelegates();
 }
 
 void FCurveEditorLatticeTool::Tick(const FGeometry& InViewContainerGeometry, const double InCurrentTime, const float InDeltaTime)
@@ -118,24 +111,31 @@ FReply FCurveEditorLatticeTool::OnMouseButtonDoubleClick(
 	)
 {
 	CancelDrag();
-	const bool bHandled = ComputeHoverStateForLeftClick(InViewContainerGeometry, InMouseEvent,
-		[this](const FLatticeHoverState& HoverState, const FLatticeBounds&)
+	
+	const bool bHandled = ComputeHoverStateForLeftClick(
+		InViewContainerGeometry, InMouseEvent,
+		[this](const FLatticeHoverState& HoverState, const FLatticeBounds&, const FCurveEditor& InCurveEditor)
 		{
 			// Double clicking the center will mirror everything with the midpoint being a horizontal line through center of the grid.
 			if (HoverState.HoveredCell)
 			{
 				const FScopedTransaction Transaction(LOCTEXT("FlipCurves", "Flip curves"));
-				MirrorGrid();
+				UndoObject->Modify();
+				HandleMirrorGrid();
+				SaveUndoState(InCurveEditor);
 				return true;
 			}
 			if (HoverState.HoveredEdge)
 			{
 				const FScopedTransaction Transaction(LOCTEXT("FlattenCurves", "Flatten curves"));
-				MoveEdgeToOpposite(*HoverState.HoveredEdge);
+				UndoObject->Modify();
+				HandleMoveEdgeToOpposite(*HoverState.HoveredEdge);
+				SaveUndoState(InCurveEditor);
 				return true;
 			}
 			return false;
 		});
+	
 	return bHandled
 		// Capturing the mouse is important so OnMouseButtonDown and OnMouseButtonUp are called on us instead of bubbling them up the hierarchy.
 		// Without this, SInteractiveCurveEditorView::OnMouseButtonUp can cause a curve to be clicked and select all keys under it.
@@ -149,7 +149,7 @@ FReply FCurveEditorLatticeTool::OnMouseButtonDown(
 {
 	CancelDrag();
 	const bool bHandled = ComputeHoverStateForLeftClick(InViewContainerGeometry, InMouseEvent,
-		[this, &InMouseEvent](const FLatticeHoverState& HoverState, const FLatticeBounds& Bounds)
+		[this, &InMouseEvent](const FLatticeHoverState& HoverState, const FLatticeBounds& Bounds, const FCurveEditor&)
 		{
 			if (HoverState.IsHovered())
 			{
@@ -204,10 +204,7 @@ FReply FCurveEditorLatticeTool::OnMouseMove(
 		}
 		else if (DelayedDragData->DelayedDrag.AttemptDragStart(InMouseEvent))
 		{
-			StartDragDrop(DelayedDragData->Bounds, DelayedDragData->HoverState);
-			DelayedDragData->ActiveOperation->DragOp->BeginDrag(
-				InViewContainerGeometry, InMouseEvent, DelayedDragData->DelayedDrag.GetInitialPosition()
-				);
+			StartDragDrop(InViewContainerGeometry, InMouseEvent);
 		}
 			
 		// Steal the capture, as we're now the authoritative widget in charge of a mouse-drag operation
@@ -231,6 +228,7 @@ void FCurveEditorLatticeTool::OnToolActivated()
 	if (const TSharedPtr<FCurveEditor> CurveEditor = WeakCurveEditor.Pin())
 	{
 		UpdateDeformerBounds(CurveEditor.ToSharedRef());
+		CurveEditor->GetTransactionManager()->OnCommandAppended().AddRaw(this, &FCurveEditorLatticeTool::OnUndoCommandAppended);
 	}
 	
 	if (UTransBuffer* Transactor = GEditor ? Cast<UTransBuffer>(GEditor->Trans) : nullptr)
@@ -243,7 +241,7 @@ void FCurveEditorLatticeTool::OnToolDeactivated()
 {
 	// Clear data. We'll recompute a new grid next time the user activate the tool
 	LatticeDeformerState.Reset();
-	UnregisterTransactionStateChanged();
+	UnregisterDelegates();
 }
 
 void FCurveEditorLatticeTool::BindCommands(TSharedRef<FUICommandList> CommandBindings)
@@ -281,9 +279,27 @@ void FCurveEditorLatticeTool::AddReferencedObjects(FReferenceCollector& Collecto
 }
 
 FCurveEditorLatticeTool::FLatticeDeformerState::FLatticeDeformerState(const FLatticeBounds& Lattice, const TSharedRef<FCurveEditor>& InCurveEditor)
-	: Deformer(1, 1, Lattice.MinValues, Lattice.MaxValues)
+	: ControlPointToCurveSpace(
+		TransformRectBetweenSpaces(Lattice.MinValues, Lattice.MaxValues, Lattice.MinValuesCurveSpace, Lattice.MaxValuesCurveSpace)
+		)
+	, GlobalDeformer(1, 1, Lattice.MinValues, Lattice.MaxValues)
+	, PerCurveData(BuildPerLatticeData(Lattice, *InCurveEditor))
 	, CurveChangeListener(InCurveEditor)
+	, PanelRebuildListener(InCurveEditor)
 {}
+	
+TArray<FVector2D> FCurveEditorLatticeTool::FLatticeDeformerState::TransformControlPointsToCurveSpace() const
+{
+	const TConstArrayView<FVector2D> ControlPoints = GlobalDeformer.GetControlPoints();
+	TArray<FVector2D> Result;
+	Result.Reserve(ControlPoints.Num());
+	Algo::Transform(ControlPoints, Result, [this](const FVector2D& ControlPoint)
+	{
+		const FVector2D Transformed = ControlPointToCurveSpace.TransformPoint(ControlPoint);
+		return Transformed;
+	});
+	return Result;
+}
 
 void FCurveEditorLatticeTool::UpdateDeformerBounds(const TSharedRef<FCurveEditor>& InCurveEditor, ELatticeUpdateFlags InRefreshFlags)
 {
@@ -301,21 +317,13 @@ void FCurveEditorLatticeTool::UpdateDeformerBounds(const TSharedRef<FCurveEditor
 	// Handle the edge case of all points overlapping, or a single key selection.
 	// The lattice math does not work when widht or height are 0 due to division by 0; fyi, the deformer handles it with an ensure.
 	const FLatticeBounds Bounds = ComputeBounds(*InCurveEditor);
-	const FVector2D DeltaValues = Bounds.MaxValues - Bounds.MinValues;
-	const bool bTooSmall = DeltaValues.X < UE_KINDA_SMALL_NUMBER || DeltaValues.Y < UE_KINDA_SMALL_NUMBER || InCurveEditor->Selection.Count() <= 1;
-	if (Bounds.bIsVisible && !bTooSmall)
+	if (Bounds.bIsVisible && !IsLatticeTooSmall(Bounds.MinValues, Bounds.MaxValues) && InCurveEditor->Selection.Count() > 1)
 	{
 		SelectionSerialNumber = InCurveEditor->Selection.GetSerialNumber();
 		LatticeDeformerState.Emplace(Bounds, InCurveEditor);
-		AddSelectionToDeformer(GetDeformer(), *InCurveEditor);
 		LatticeDeformerState->CurveChangeListener.OnCurveModified().AddRaw(this, &FCurveEditorLatticeTool::OnCurvesModified);
-
-		const bool bActuallyTransact = !EnumHasAnyFlags(InRefreshFlags, ELatticeUpdateFlags::NoTransaction)
-			// Re-use ongoing transaction if currently recording into buffer
-			&& GUndo == nullptr && !GIsTransacting; 
-		FScopedTransaction Transaction(TEXT("CurveEditorLatticeTool"), LOCTEXT("SetSelection", "Set lattice selection"), nullptr, bActuallyTransact);
-		UndoObject->Modify();
-		UndoObject->SaveState(*InCurveEditor, GetDeformer());
+		LatticeDeformerState->PanelRebuildListener.OnCurveViewRebuilt().AddRaw(this, &FCurveEditorLatticeTool::OnCurvePanelRebuilt);
+		SaveUndoState(*InCurveEditor);
 	}
 }
 
@@ -337,13 +345,40 @@ void FCurveEditorLatticeTool::OnCurvesModified()
 	LatticeDeformerState->bHasRequestedRefresh = true;
 }
 
+void FCurveEditorLatticeTool::OnUndoCommandAppended(const CurveEditor::FCurveEditorCommandChange& InCommand)
+{
+	if (const TSharedPtr<FCurveEditor> CurveEditor = WeakCurveEditor.Pin())
+	{
+		// This will diff the selection serial number and modify the undo object.
+		UpdateDeformerBounds(CurveEditor.ToSharedRef(), ELatticeUpdateFlags::NoTransaction);
+	}
+}
+
 void FCurveEditorLatticeTool::OnTransactionStateChanged(
 	const FTransactionContext& Context, ETransactionStateEventType TransactionStateEvent
 	)
 {
-	// Technically, we should check whether the transaction actually affected the lattice... but this isn't computationally expensive so whatever.
-	const bool bNeedsRecompute_UndoRedo = TransactionStateEvent == ETransactionStateEventType::UndoRedoFinalized;
-	if (bNeedsRecompute_UndoRedo)
+	UTransBuffer* TransBuffer = GEditor ? CastChecked<UTransBuffer>(GEditor->Trans) : nullptr;
+	if (!TransBuffer || TransactionStateEvent != ETransactionStateEventType::UndoRedoFinalized)
+	{
+		return;
+	}
+	const TSharedRef<FTransaction>* CurrentTransaction = TransBuffer->UndoBuffer.FindByPredicate(
+		[&Context](const TSharedRef<FTransaction>& Transaction)
+		{
+			return Transaction->GetOperationId() == Context.OperationId;
+		});
+	
+	// We're trying to handle cases where curves are changed but OnCurvesModified, i.e. FCurveModel::OnCurveModified, is not triggered...
+	// This happens e.g. when underlying Sequencer data changes, which is under the hood of the FCurveModel interface.
+	// We'll just always regenerat the grid, i.e. a new rectangular lattice will be placed around the selection.
+	//
+	// Sadly, this means that undos / redos on totally unrelated places in the engine (e.g. undo an edit operation in details panel), will cause the
+	// lattice to be regenerated, too. This is fine as long as it does not happen while the user is actively working in the curve editor.
+	// This is why we check whether the UndoObject is in the current transaction... in that case we do NOT recompute the lattice since our custom
+	// undo logic may change control points directly.
+	const bool bAppliesToLattice = CurrentTransaction && CurrentTransaction->Get().ContainsObject(UndoObject);
+	if (!bAppliesToLattice)
 	{
 		RecomputeLatticeGrid();
 	}
@@ -357,21 +392,12 @@ void FCurveEditorLatticeTool::RecomputeLatticeGrid(const TSharedRef<FCurveEditor
 	LatticeDeformerState.Reset();
 	UpdateDeformerBounds(InCurveEditor, InRefreshFlags);
 }
-
-void FCurveEditorLatticeTool::CancelDrag()
-{
-	if (DelayedDragData && DelayedDragData->ActiveOperation)
-	{
-		DelayedDragData->ActiveOperation->DragOp->CancelDrag();
-	}
-	DelayedDragData.Reset();
-}
 	
 void FCurveEditorLatticeTool::RecomputeDrawData(const FGeometry& InViewContainerGeometry, const TSharedPtr<FCurveEditor>& InCurveEditor)
 {
 	// The widget hierarchy is SCurveEditorViewContainer -> SCurveEditorViewAbsolute.
 	// InAllottedGeometry is that of SCurveEditorViewContainer.
-	DrawData->MatrixWidth = GetDeformer().NumPointsInWidth();
+	DrawData->MatrixWidth = GetGlobalDeformer().NumPointsInWidth();
 	DrawData->ControlPoints = ComputeDrawControlPoints(*InCurveEditor);
 	const bool bHasControlData = !DrawData->ControlPoints.IsEmpty();
 	if (!bHasControlData)
@@ -386,7 +412,7 @@ void FCurveEditorLatticeTool::RecomputeDrawData(const FGeometry& InViewContainer
 		DrawData->HoverState = ComputeLatticeHoverState(
 			InViewContainerGeometry, FinalMouseScreenSpacePosition,
 			DrawData->ControlPoints,DrawData->ControlEdges,
-			GetDeformer().NumPointsInWidth(), GetDeformer().NumCells()
+			GetGlobalDeformer().NumPointsInWidth(), GetGlobalDeformer().NumCells()
 		);
 	}
 }
@@ -395,8 +421,8 @@ TArray<FVector2D> FCurveEditorLatticeTool::ComputeDrawControlPoints(const FCurve
 {
 	check(LatticeDeformerState);
 	TArray<FVector2D> Result;
-
-	const TConstArrayView<FVector2D> ControlPoints = GetDeformer().GetControlPoints();
+	
+	const TArray<FVector2D> ControlPoints = LatticeDeformerState->TransformControlPointsToCurveSpace();
 	Result.SetNumUninitialized(ControlPoints.Num());
 	if (!TransformViewToViewContainer(InCurveEditor, ControlPoints, Result))
 	{
@@ -408,115 +434,327 @@ TArray<FVector2D> FCurveEditorLatticeTool::ComputeDrawControlPoints(const FCurve
 TArray<FLatticeControlEdge> FCurveEditorLatticeTool::ComputeDrawControlEdges(TConstArrayView<FVector2D> InControlPoints) const
 {
 	TArray<FLatticeControlEdge> ControlEdges;
-	ControlEdges.SetNumUninitialized(GetDeformer().NumEdges());
-	for (int32 EdgeIndex = 0; EdgeIndex < GetDeformer().NumEdges(); ++EdgeIndex)
+	ControlEdges.SetNumUninitialized(GetGlobalDeformer().NumEdges());
+	for (int32 EdgeIndex = 0; EdgeIndex < GetGlobalDeformer().NumEdges(); ++EdgeIndex)
 	{
-		const auto[Start, End] = GetDeformer().GetEdgeIndices(EdgeIndex);
+		const auto[Start, End] = GetGlobalDeformer().GetEdgeIndices(EdgeIndex);
 		ControlEdges[EdgeIndex] = FLatticeControlEdge(InControlPoints[Start], InControlPoints[End]);
 	}
 	return ControlEdges;
 }
 
-void FCurveEditorLatticeTool::StartDragDrop(const FLatticeBounds& InBounds, const FLatticeHoverState& InHoverState)
+static FText LexViewMode(ECurveEditorViewID ViewMode)
 {
-	const TConstArrayView<FVector2D> ControlPoints = GetDeformer().GetControlPoints();
-	if (InHoverState.HoveredControlPoint)
+	switch (ViewMode)
 	{
-		const int32 ControlPointIndex = *InHoverState.HoveredControlPoint;
-		TUniquePtr<FLatticeDragOp> Op = MakeUnique<FLatticeDragOp_MoveControlPoints>(
-			WeakCurveEditor,
-			FControlPointArray{ ControlPoints[ControlPointIndex] },
-			InBounds,
-			FLatticeDragOp_MoveControlPoints::FOnControlPointsMoved::CreateLambda(
-				[this, ControlPointIndex](const TConstArrayView<FVector2D>& InNewControlPointPositions)
-				{ OnMoveControlPoint({ ControlPointIndex }, InNewControlPointPositions, EPropertyChangeType::Interactive); }
-			),
-			FLatticeDragOp_MoveControlPoints::FOnControlPointsMoved::CreateLambda(
-				[this, ControlPointIndex](const TConstArrayView<FVector2D>& InNewControlPointPositions)
-				{ OnMoveControlPoint({ ControlPointIndex }, InNewControlPointPositions, EPropertyChangeType::ValueSet); }
-			)
+	case ECurveEditorViewID::Absolute: return LOCTEXT("Absolute", "Absolute");
+	case ECurveEditorViewID::Normalized: return LOCTEXT("Normalized", "Normalized");
+	case ECurveEditorViewID::Stacked: return LOCTEXT("Stacked", "Stacked");
+	default: return LOCTEXT("Custom", "Custom");
+	}
+}
+	
+void FCurveEditorLatticeTool::CancelDrag()
+{
+	if (DelayedDragData && DelayedDragData->ActiveOperation)
+	{
+		DelayedDragData->ActiveOperation->DragOp->CancelDrag();
+	}
+	DelayedDragData.Reset();
+}
+	
+void FCurveEditorLatticeTool::StartDragDrop(const FGeometry& InViewContainerGeometry, const FPointerEvent& InMouseEvent)
+{
+	if (TOptional<FDragOpData> DragDrop = CreateDragDrop(DelayedDragData->Bounds, DelayedDragData->HoverState))
+	{
+		DelayedDragData->ActiveOperation = MoveTemp(*DragDrop);
+		DelayedDragData->ActiveOperation->DragOp->BeginDrag(
+			InViewContainerGeometry, InMouseEvent, DelayedDragData->DelayedDrag.GetInitialPosition()
 		);
-		DelayedDragData->ActiveOperation.Emplace(MoveTemp(Op), LOCTEXT("MoveControlPoint", "Move control point"));
-		DelayedDragData->PointDragTangentMirrorOperation = FLatticePointTangentsMirrorOp::MakeFromDraggedControlPoint(
-			ControlPointIndex, GetDeformer(), *WeakCurveEditor.Pin()
-			);
 	}
-
+}
+	
+TOptional<FCurveEditorLatticeTool::FDragOpData> FCurveEditorLatticeTool::CreateDragDrop(
+	const FLatticeBounds& InBounds, const FLatticeHoverState& InHoverState
+	)
+{
+	if (InHoverState.HoveredControlPoint)
+	{
+		return CreateDragDrop_Point(InBounds, *InHoverState.HoveredControlPoint);
+	}
 	if (InHoverState.HoveredEdge)
 	{
-		const int32 EdgeIndex = *InHoverState.HoveredEdge;
-		const auto[Start, End] = GetDeformer().GetEdgeIndices(EdgeIndex);
-		TUniquePtr<FLatticeDragOp> Op = MakeUnique<FLatticeDragOp_MoveControlPoints>(
-			WeakCurveEditor,
-			FControlPointArray{ ControlPoints[Start], ControlPoints[End] },
-			InBounds,
-			FLatticeDragOp_MoveControlPoints::FOnControlPointsMoved::CreateLambda(
-				[this, Start, End](const TConstArrayView<FVector2D>& InNewControlPointPositions)
-				{ OnMoveEdge({ Start, End }, InNewControlPointPositions, EPropertyChangeType::Interactive); }
-			),
-			FLatticeDragOp_MoveControlPoints::FOnControlPointsMoved::CreateLambda(
-				[this, Start, End](const TConstArrayView<FVector2D>& InNewControlPointPositions)
-				{ OnMoveEdge({ Start, End }, InNewControlPointPositions, EPropertyChangeType::ValueSet); }
-			)
+		return CreateDragDrop_Edge(InBounds, *InHoverState.HoveredEdge);
+	}
+	if (InHoverState.HoveredCell)
+	{
+		return CreateDragDrop_Cell(InBounds, *InHoverState.HoveredCell);
+	}
+
+	return {};
+}
+
+TOptional<FCurveEditorLatticeTool::FDragOpData> FCurveEditorLatticeTool::CreateDragDrop_Point(
+	const FLatticeBounds& InBounds, const int32 InControlPointIndex
+	)
+{
+	using FHandleMovedPoints = FLatticeDragOp_MoveControlPoints::FOnControlPointsMoved;
+	const TConstArrayView<FVector2D> ControlPoints = GetGlobalDeformer().GetControlPoints();
+	
+	const auto MakeDelegate = [this, InControlPointIndex](EPropertyChangeType::Type ChangeType)
+	{
+		return FHandleMovedPoints::CreateLambda(
+			[this, InControlPointIndex, ChangeType](const TConstArrayView<FVector2D>& InNewControlPointPositions)
+			{
+				const int32 Indices[1] { InControlPointIndex };
+				HandleDragPoints(&FCurveEditorLatticeTool::OnMoveControlPoint,Indices, InNewControlPointPositions, ChangeType);
+				GetGlobalDeformer().UpdateControlPoints(Indices, InNewControlPointPositions, [](auto, auto){});
+
+				// When the drag is finished, save where the lattice is so its shape can be transacted correctly.
+				const bool bIsDragFinished = ChangeType == EPropertyChangeType::ValueSet; 
+				if (bIsDragFinished)
+				{
+					SaveUndoState(*WeakCurveEditor.Pin());
+				}
+			});
+	};
+	const FHandleMovedPoints HandleDrag = MakeDelegate(EPropertyChangeType::Interactive);
+	const FHandleMovedPoints HandleFinishDrag = MakeDelegate(EPropertyChangeType::ValueSet);
+	
+	TUniquePtr<FLatticeDragOp> Op = MakeUnique<FLatticeDragOp_MoveControlPoints>(
+		WeakCurveEditor, FControlPointArray{ ControlPoints[InControlPointIndex] }, InBounds, HandleDrag, HandleFinishDrag
+	);
+	FDragOpData DragDrop(
+		MoveTemp(Op),
+		FText::Format(LOCTEXT("MoveControlPoint", "Move control point ({0})"), LexViewMode(WeakCurveEditor.Pin()->GetPanel()->GetViewMode()))
 		);
-		DelayedDragData->ActiveOperation.Emplace(MoveTemp(Op), LOCTEXT("MoveEdge", "Move edge"));
-		DelayedDragData->EdgeDragTangentMirrorOperation = FLatticeEdgeTangentsMirrorOp::MakeFromDraggedEdge(
-			EdgeIndex, GetDeformer(), *WeakCurveEditor.Pin()
-			);
+
+	const FCurveEditor& CurveEditor = *WeakCurveEditor.Pin();
+	for (const TPair<FCurveModelID, FPerCurveLatticeData>& Pair : LatticeDeformerState->PerCurveData)
+	{
+		const FPerCurveDeformer2D& Deformer = Pair.Value.CurveDeformer;
+		TOptional<FLatticePointTangentsMirrorOp> MirrorOp = FLatticePointTangentsMirrorOp::MakeMirrorOpForDragLatticeControlPoint(
+			InControlPointIndex, Pair.Key, Deformer, CurveEditor
+		);
+		if (MirrorOp)
+		{
+			DragDrop.PointTangentMirroringOps.Add(Pair.Key, MoveTemp(*MirrorOp));
+		}
 	}
+	return DragDrop;
+}
 
-	if (InHoverState.HoveredCell)
+TOptional<FCurveEditorLatticeTool::FDragOpData> FCurveEditorLatticeTool::CreateDragDrop_Edge(
+	const FLatticeBounds& InBounds, int32 InEdgeIndex
+	)
+{
+	using FHandleMovedPoints = FLatticeDragOp_MoveControlPoints::FOnControlPointsMoved;
+	const TConstArrayView<FVector2D> ControlPoints = GetGlobalDeformer().GetControlPoints();
+	const FEdgeVertexIndices EdgeIndices = GetGlobalDeformer().GetEdgeIndices(InEdgeIndex);
+	const auto MakeDelegate = [this, EdgeIndices](EPropertyChangeType::Type ChangeType)
 	{
-		const auto[TopLeft, TopRight, BottomRight, BottomLeft] = GetDeformer().GetCellIndices(*InHoverState.HoveredCell);
-		TUniquePtr<FLatticeDragOp> Op = MakeUnique<FLatticeDragOp_MoveControlPoints>(
-			WeakCurveEditor,
-			FControlPointArray{ ControlPoints[TopLeft], ControlPoints[TopRight], ControlPoints[BottomRight], ControlPoints[BottomLeft] },
-			InBounds,
-			FLatticeDragOp_MoveControlPoints::FOnControlPointsMoved::CreateLambda(
-				[this, TopLeft, TopRight, BottomRight, BottomLeft](const TConstArrayView<FVector2D>& InNewControlPointPositions)
+		return FHandleMovedPoints::CreateLambda(
+			[this, EdgeIndices, ChangeType](const TConstArrayView<FVector2D>& InNewControlPointPositions)
+			{
+				const int32 Indices[2]{ EdgeIndices.Start, EdgeIndices.End };
+				HandleDragPoints(&FCurveEditorLatticeTool::OnMoveEdge, Indices, InNewControlPointPositions, ChangeType);
+				GetGlobalDeformer().UpdateControlPoints(Indices, InNewControlPointPositions, [](auto, auto){});
+
+				// When the drag is finished, save where the lattice is so its shape can be transacted correctly.
+				const bool bIsDragFinished = ChangeType == EPropertyChangeType::ValueSet; 
+				if (bIsDragFinished)
 				{
-					OnSetControlPoints({ TopLeft, TopRight, BottomRight, BottomLeft }, InNewControlPointPositions, EPropertyChangeType::Interactive);
+					SaveUndoState(*WeakCurveEditor.Pin());
 				}
-			),
-			FLatticeDragOp_MoveControlPoints::FOnControlPointsMoved::CreateLambda(
-				[this, TopLeft, TopRight, BottomRight, BottomLeft](const TConstArrayView<FVector2D>& InNewControlPointPositions)
+			});
+	};
+	const FHandleMovedPoints HandleDrag = MakeDelegate(EPropertyChangeType::Interactive);
+	const FHandleMovedPoints HandleFinishDrag = MakeDelegate(EPropertyChangeType::ValueSet);
+		
+	FControlPointArray EdgePoints{ ControlPoints[EdgeIndices.Start], ControlPoints[EdgeIndices.End] };
+	TUniquePtr<FLatticeDragOp> Op = MakeUnique<FLatticeDragOp_MoveControlPoints>(
+		WeakCurveEditor, MoveTemp(EdgePoints), InBounds, HandleDrag, HandleFinishDrag
+	);
+	FDragOpData DragDrop(
+		MoveTemp(Op),
+		FText::Format(LOCTEXT("MoveEdge", "Move edge ({0})"), LexViewMode(WeakCurveEditor.Pin()->GetPanel()->GetViewMode()))
+		);
+
+	const FCurveEditor& CurveEditor = *WeakCurveEditor.Pin();
+	for (const TPair<FCurveModelID, FPerCurveLatticeData>& Pair : LatticeDeformerState->PerCurveData)
+	{
+		const FPerCurveDeformer2D& Deformer = Pair.Value.CurveDeformer;
+		TOptional<FLatticeEdgeTangentsMirrorOp> MirrorOp = FLatticeEdgeTangentsMirrorOp::MakeMirrorOpForDragLatticeEdge(
+			InEdgeIndex, Pair.Key, Deformer, CurveEditor
+			);
+		if (MirrorOp)
+		{
+			DragDrop.EdgeTangentMirroringOps.Emplace(Pair.Key, MoveTemp(*MirrorOp));
+		}
+	}
+	return DragDrop;
+}
+
+TOptional<FCurveEditorLatticeTool::FDragOpData> FCurveEditorLatticeTool::CreateDragDrop_Cell(
+	const FLatticeBounds& InBounds, int32 InCellIndex
+	)
+{
+	using FHandleMovedPoints = FLatticeDragOp_MoveControlPoints::FOnControlPointsMoved;
+	const TConstArrayView<FVector2D> ControlPoints = GetGlobalDeformer().GetControlPoints();
+	
+	const FCellVertexIndices CellIndices = GetGlobalDeformer().GetCellIndices(InCellIndex);
+	const auto MakeDelegate = [this, CellIndices](EPropertyChangeType::Type ChangeType)
+	{
+		return FHandleMovedPoints::CreateLambda(
+			[this, CellIndices, ChangeType](const TConstArrayView<FVector2D>& InNewControlPointPositions)
+			{
+				const int32 Indices[4]{ CellIndices.TopLeft, CellIndices.TopRight, CellIndices.BottomRight, CellIndices.BottomLeft };
+				HandleDragPoints(&FCurveEditorLatticeTool::OnSetControlPoints, Indices, InNewControlPointPositions, ChangeType);
+				GetGlobalDeformer().UpdateControlPoints(Indices, InNewControlPointPositions, [](auto, auto){});
+				
+				// When the drag is finished, save where the lattice is so its shape can be transacted correctly.
+				const bool bIsDragFinished = ChangeType == EPropertyChangeType::ValueSet; 
+				if (bIsDragFinished)
 				{
-					OnSetControlPoints({ TopLeft, TopRight, BottomRight, BottomLeft }, InNewControlPointPositions, EPropertyChangeType::ValueSet);
+					SaveUndoState(*WeakCurveEditor.Pin());
 				}
-			)
+			});
+	};
+	const FHandleMovedPoints HandleDrag = MakeDelegate(EPropertyChangeType::Interactive);
+	const FHandleMovedPoints HandleFinishDrag = MakeDelegate(EPropertyChangeType::ValueSet);
+		
+	const auto[TopLeft, TopRight, BottomRight, BottomLeft] = CellIndices;
+	FControlPointArray CellPoints{ ControlPoints[TopLeft], ControlPoints[TopRight], ControlPoints[BottomRight], ControlPoints[BottomLeft] };
+	TUniquePtr<FLatticeDragOp> Op = MakeUnique<FLatticeDragOp_MoveControlPoints>(
+		WeakCurveEditor, MoveTemp(CellPoints), InBounds, HandleDrag, HandleFinishDrag
+	);
+	return FDragOpData(
+		MoveTemp(Op),
+		FText::Format(LOCTEXT("MoveCell", "Move cell ({0})"), LexViewMode(WeakCurveEditor.Pin()->GetPanel()->GetViewMode()))
 		);
-		DelayedDragData->ActiveOperation.Emplace(MoveTemp(Op), LOCTEXT("MoveCell", "Move cell"));
+}
+
+void FCurveEditorLatticeTool::HandleDragPoints(
+	FMovePointsSignature InMovePointsCallback, TConstArrayView<int32> InVertexIndices,
+	TConstArrayView<FVector2D> InNewControlPoints, EPropertyChangeType::Type InChangeType
+	)
+{
+	const TSharedPtr<FCurveEditor> CurveEditorPin = WeakCurveEditor.Pin();
+	if (!CurveEditorPin)
+	{
+		return;
+	}
+
+	if (!DelayedDragData->ActiveOperation->bSavedUndoState)
+	{
+		DelayedDragData->ActiveOperation->bSavedUndoState = true;
+		UndoObject->Modify();
+	}
+	
+	for (TPair<FCurveModelID, FPerCurveLatticeData>& Pair : LatticeDeformerState->PerCurveData)
+	{
+		FCurveModel* CurveModel = CurveEditorPin->FindCurve(Pair.Key);
+		if (!CurveModel)
+		{
+			continue;
+		}
+
+		// InNewControlPoints are in absolute space (literal FKeyPosition::InputValue / OutputValue).
+		// Each curve's lattice is still in absolute space but depending on the view, the lattice may have been adjusted to a smaller area.
+		// We need to transform the dragged control points to the space used by the per-curve lattice.
+		const TArray<FVector2D> CurveRelativeControlPoints = TransformPoints(Pair.Value.GlobalDeformerToCurveDeformer, InNewControlPoints);
+		
+		FPerCurveDeformer2D& PerCurveDeformer = Pair.Value.CurveDeformer;
+		(this->*InMovePointsCallback)(PerCurveDeformer, *CurveModel, InVertexIndices, CurveRelativeControlPoints, InChangeType);
+	}
+}
+
+void FCurveEditorLatticeTool::HandleMoveEdgeToOpposite(int32 InEdgeIndex)
+{
+	const TSharedPtr<FCurveEditor> CurveEditorPin = WeakCurveEditor.Pin();
+	const TOptional<FMoveEdgeToOppositeData> MoveData = ComputeMoveEdgeToOppositeData(GetGlobalDeformer(), InEdgeIndex);
+	if (!CurveEditorPin || !MoveData)
+	{
+		return;
+	}
+	
+	// Move the per-curve deformers.
+	for (TPair<FCurveModelID, FPerCurveLatticeData>& Pair : LatticeDeformerState->PerCurveData)
+	{
+		if (FCurveModel* Curve = CurveEditorPin->FindCurve(Pair.Key))
+		{
+			MoveEdgeToOpposite(Pair.Value.CurveDeformer, *Curve, *CurveEditorPin, InEdgeIndex);
+		}
 	}
+
+	// Move the UI only deformer.
+	const auto NoOp = [](auto, auto){};
+	GetGlobalDeformer().UpdateControlPoints({ MoveData->EdgeIndices.Start, MoveData->EdgeIndices.End }, MoveData->EdgeControlPoints.Points, NoOp);
 }
 
-void FCurveEditorLatticeTool::OnMoveControlPoint(TConstArrayView<int32> InVertexIndices, TConstArrayView<FVector2D> InNewControlPointPositions, EPropertyChangeType::Type InChangeType)
+void FCurveEditorLatticeTool::HandleMirrorGrid()
+{
+	const TSharedPtr<FCurveEditor> CurveEditorPin = WeakCurveEditor.Pin();
+	if (!CurveEditorPin)
+	{
+		return;
+	}
+	
+	// 1. Mirror all the keys
+	for (TPair<FCurveModelID, FPerCurveLatticeData>& Pair : LatticeDeformerState->PerCurveData)
+	{
+		if (FCurveModel* Curve = CurveEditorPin->FindCurve(Pair.Key))
+		{
+			MirrorGrid(Pair.Value.CurveDeformer, *Curve, *CurveEditorPin);
+		}
+	}
+	
+	// 2. Recompute lattice grid
+	// This will calculate a new rectangular grid around the points. Instead, we could try to preserve the existing shape... but nobody cares, so lets not implement it.
+	LatticeDeformerState.Reset();
+	UpdateDeformerBounds(CurveEditorPin.ToSharedRef());
+}
+
+void FCurveEditorLatticeTool::OnMoveControlPoint(
+	FPerCurveDeformer2D& InDeformer, FCurveModel& InCurveModel,
+	TConstArrayView<int32> InVertexIndices, TConstArrayView<FVector2D> InNewControlPointPositions,
+	EPropertyChangeType::Type InChangeType
+	)
 {
 	check(InVertexIndices.Num() == 1 && InNewControlPointPositions.Num() == 1);
 	TGuardValue Guard(LatticeDeformerState->bIsModifyingCurves, true);
 	
-	if (DelayedDragData->PointDragTangentMirrorOperation)
+	const FCurveModelID& CurveId = *InCurveModel.GetId();
+	if (FLatticePointTangentsMirrorOp* MirrorOp = DelayedDragData->ActiveOperation->PointTangentMirroringOps.Find(CurveId))
 	{
-		DelayedDragData->PointDragTangentMirrorOperation->OnMovePoint(InNewControlPointPositions[0]);
+		MirrorOp->OnMovePoint(InNewControlPointPositions[0], *WeakCurveEditor.Pin());
 	}
 	
-	OnSetControlPoints(InVertexIndices, InNewControlPointPositions, InChangeType);
+	OnSetControlPoints(InDeformer, InCurveModel, InVertexIndices, InNewControlPointPositions, InChangeType);
 }
 
-void FCurveEditorLatticeTool::OnMoveEdge(TConstArrayView<int32> InVertexIndices, TConstArrayView<FVector2D> InNewControlPointPositions, EPropertyChangeType::Type InChangeType)
+void FCurveEditorLatticeTool::OnMoveEdge(
+	FPerCurveDeformer2D& InDeformer, FCurveModel& InCurveModel,
+	TConstArrayView<int32> InVertexIndices, TConstArrayView<FVector2D> InNewControlPointPositions,
+	EPropertyChangeType::Type InChangeType
+	)
 {
 	check(InVertexIndices.Num() == 2 && InNewControlPointPositions.Num() == 2);
 	TGuardValue Guard(LatticeDeformerState->bIsModifyingCurves, true);
-	
-	if (DelayedDragData->EdgeDragTangentMirrorOperation)
+
+	const FCurveModelID& CurveId = *InCurveModel.GetId();
+	if (FLatticeEdgeTangentsMirrorOp* MirrorOp = DelayedDragData->ActiveOperation->EdgeTangentMirroringOps.Find(CurveId))
 	{
-		DelayedDragData->EdgeDragTangentMirrorOperation->OnMoveEdge(InNewControlPointPositions);
+		MirrorOp->OnMoveEdge(InNewControlPointPositions, *WeakCurveEditor.Pin());
 	}
 
-	OnSetControlPoints(InVertexIndices, InNewControlPointPositions, InChangeType);
+	OnSetControlPoints(InDeformer, InCurveModel, InVertexIndices, InNewControlPointPositions, InChangeType);
 }
 
 void FCurveEditorLatticeTool::OnSetControlPoints(
-	TConstArrayView<int32> InControlPointsToUpdate, TConstArrayView<FVector2D> InNewControlPoints, EPropertyChangeType::Type InChangeType
+	FPerCurveDeformer2D& InDeformer, FCurveModel& InCurveModel,
+	TConstArrayView<int32> InControlPointsToUpdate, TConstArrayView<FVector2D> InNewControlPoints,
+	EPropertyChangeType::Type InChangeType
 	)
 {
 	const TSharedPtr<FCurveEditor> CurveEditor = WeakCurveEditor.Pin();
@@ -530,121 +768,80 @@ void FCurveEditorLatticeTool::OnSetControlPoints(
 	
 	// Combining Handles & Keys into large arrays like this is significantly faster than feeding each key into SetKeyPositions separately.
 	// For 9000 keys it goes from 2136 ms to 238 ms.
-	struct FCurveChanges
-	{
-		TArray<FKeyHandle> Handles;
-		TArray<FKeyPosition> Keys;
-	};
-
-	TMap<FCurveModelID, FCurveChanges> Changes;
-	// Pre-allocating saves about 0.060ms
-	for (const TPair<FCurveModelID, FKeyHandleSet>& Pair : CurveEditor->Selection.GetAll())
-	{
-		FCurveChanges& CurveChanges = Changes.Add(Pair.Key);
-		CurveChanges.Handles.Reserve(Pair.Value.Num());
-		CurveChanges.Keys.Reserve(Pair.Value.Num());
-	}
+	TArray<FKeyHandle> Handles;
+	TArray<FKeyPosition> Keys;
+	
+	// Pre-allocating saves about 0.060ms (also it is needed because UpdateControlPoints is executed in parallel).
+	const int32 NumKeys = InDeformer.NumKeys();
+	Handles.Reserve(NumKeys);
+	Keys.Reserve(NumKeys);
 	
-	// The actual re-computation takes about 0.083 ms
-	GetDeformer().UpdateControlPoints(InControlPointsToUpdate, InNewControlPoints,
-		[&Snapper, &Changes](const FKeyInCurveId& Handle, const FVector2D& KeyPosition)
+	InDeformer.UpdateControlPoints(InControlPointsToUpdate, InNewControlPoints,
+		[&Snapper, &Handles, &Keys](const FKeyHandle& KeyHandle, const FVector2D& KeyPosition)
 		{
-			FCurveChanges& CurveChanges = Changes.FindOrAdd(Handle.CurveId);
-			CurveChanges.Handles.Add(Handle.KeyId);
-			CurveChanges.Keys.Add(
+			Handles.Add(KeyHandle);
+			Keys.Add(
 				Snapper->SnapKey(FKeyPosition{ KeyPosition.X, KeyPosition.Y })
 				);
 		});
 
-	// FDragOpData has started a FScopedTransaction and will release it once dragging has stopped.
-	UndoObject->Modify();
-	UndoObject->SaveState(*CurveEditor, GetDeformer());
-
 	// Some stats all with 9000 keys. Varying degree of changed keys, i.e. how far you move your mouse in quick succession.
 	// Sorted by little to a lot of changes:
 	// - SetKeyPositions: 4.299 ms		- Total: 4.533 ms
 	// - SetKeyPositions 138.954 ms		- Total 139.229 ms
 	// - SetKeyPositions 269.523 ms		- Total 269.811 ms
 	// Conclusion: The following SetKeyPositions block is slow.
-	for (auto It = Changes.CreateConstIterator(); It; ++It)
-	{
-		const FCurveChanges& Pair = It->Value;
-		if (FCurveModel* Model = CurveEditor->FindCurve(It->Key); ensure(Model))
-		{
-			Model->Modify();
-			Model->SetKeyPositions(Pair.Handles, Pair.Keys, InChangeType);
-		}
-	}
+	InCurveModel.Modify();
+	InCurveModel.SetKeyPositions(Handles, Keys, InChangeType);
 }
 
-void FCurveEditorLatticeTool::MoveEdgeToOpposite(int32 InEdgeIndex)
+void FCurveEditorLatticeTool::MoveEdgeToOpposite(FPerCurveDeformer2D& InDeformer, FCurveModel& InCurveModel, FCurveEditor& InCurveEditor, int32 InEdgeIndex)
 {
-	TSharedPtr<FCurveEditor> CurveEditor = WeakCurveEditor.Pin();
-	const FLatticeDeformer2D& Deformer = GetDeformer();
-	// At time of implementation, we only have a single cell. Should that change, adjust this implementation.
-	if (Deformer.NumCells() != 1 || !CurveEditor)
+	const TOptional<FMoveEdgeToOppositeData> MoveData = ComputeMoveEdgeToOppositeData(InDeformer, InEdgeIndex);
+	if (!MoveData)
 	{
 		return;
 	}
-
-	const ELatticeEdgeType EdgeType = GetEdgeTypeFromIndexInSingleCellGrid(InEdgeIndex);
-	if (EdgeType != ELatticeEdgeType::Top && EdgeType != ELatticeEdgeType::Bottom)
-	{
-		return;
-	}
-	
-	const ELatticeEdgeType OppositeEdge = GetOppositeEdge(EdgeType);
-	const int32 OppositeEdgeIndex = GetEdgeIndexInSingleCellGrid(OppositeEdge);
-	const auto[MovedEdgeStartVertIdx, MovedEdgeEndVertIdx] = Deformer.GetEdgeIndices(OppositeEdgeIndex);
-	const FLatticeControlEdge& Edge = Deformer.GetControlEdge(InEdgeIndex);
-	const TArray<FVector2D, TInlineAllocator<2>> NewEdge = { Edge.Start(), Edge.End() };
 	
 	// OnSetControlPoints modifies the curves too late - MirrorTangents will already have modified them.
-	const TConstArrayView<FKeyInCurveId> DeformerKeys = Deformer.GetKeyIdsInCell(0);
-	ModifyCurves(DeformerKeys, *CurveEditor); 
+	InCurveModel.Modify();
 	
 	// Flatten tangents if moving top or bottom edge...
-	TOptional<FLatticeEdgeTangentsMirrorOp> MirrorOp = FLatticeEdgeTangentsMirrorOp::MakeFromDraggedEdge(
-		OppositeEdgeIndex, GetDeformer(), *WeakCurveEditor.Pin()
+	const FCurveEditor& CurveEditor = *WeakCurveEditor.Pin();
+	TOptional<FLatticeEdgeTangentsMirrorOp> MirrorOp = FLatticeEdgeTangentsMirrorOp::MakeMirrorOpForDragLatticeEdge(
+		MoveData->OppositeEdgeIndex, *InCurveModel.GetId(), InDeformer, CurveEditor
 		);
-	// Even though we've checked that we're moving a top or bottom edge, the mirror op can still be unset if all keys are set to auto interpolate.
 	if (MirrorOp) 
 	{
-		MirrorOp->OnMoveEdge(NewEdge);
+		MirrorOp->OnMoveEdge(MoveData->EdgeControlPoints.Points, CurveEditor);
 	}
+	
 	// ... then move the grid
-	OnSetControlPoints({ MovedEdgeStartVertIdx, MovedEdgeEndVertIdx }, NewEdge, EPropertyChangeType::ValueSet);
+	OnSetControlPoints(
+		InDeformer, InCurveModel,
+		{ MoveData->EdgeIndices.Start, MoveData->EdgeIndices.End }, MoveData->EdgeControlPoints.Points,
+		EPropertyChangeType::ValueSet
+		);
 }
 
-void FCurveEditorLatticeTool::MirrorGrid()
+void FCurveEditorLatticeTool::MirrorGrid(FPerCurveDeformer2D& InDeformer, FCurveModel& InCurveModel, FCurveEditor& InCurveEditor)
 {
-	TSharedPtr<FCurveEditor> CurveEditor = WeakCurveEditor.Pin();
-	const FLatticeDeformer2D& Deformer = GetDeformer();
 	// At time of implementation, we only have a single cell. Should that change, adjust this implementation.
-	if (Deformer.NumCells() != 1 || !CurveEditor)
+	if (InDeformer.NumCells() != 1)
 	{
 		return;
 	}
-
-	const TConstArrayView<FKeyInCurveId> DeformerKeys = Deformer.GetKeyIdsInCell(0);
-	ModifyCurves(DeformerKeys, *CurveEditor); // OnSetControlPoints modifies the curves too late - MirrorTangents will already have modified them.
-
-	// 1. Mirror all the keys
-	const TMap<FCurveModelID, TArray<FKeyHandle>> SplitCurves = SplitKeysInCurves(DeformerKeys);
-	CurveEditor::MirrorOnMidpoint(SplitCurves, *CurveEditor, [&CurveEditor](
-		TConstArrayView<FKeyHandle> Keys, const FCurveModelID& CurveId, double Min, double Max, double Midpoint)
-		{
-			CurveEditor::MirrorTangents(Keys, CurveId, Min, Max, Midpoint, *CurveEditor);
-			CurveEditor::MirrorKeyPositions(Keys, CurveId, Min, Max, Midpoint, *CurveEditor);
-		});
-
-	// 2. Recompute lattice grid
-	// This will calculate a new rectangular grid around the points. Instead, we could try to preserve the existing shape... but nobody cares, so lets not implement it.
-	LatticeDeformerState.Reset();
-	UpdateDeformerBounds(CurveEditor.ToSharedRef());
+	
+	InCurveModel.Modify(); 
+	
+	const TConstArrayView<FKeyHandle> Keys = InDeformer.GetCellMetaData(0);
+	const auto[Min, Max] = CurveEditor::FindMinMaxHeight(Keys, InCurveModel);
+	const double Midpoint = Min + (Max - Min) * 0.5;
+	CurveEditor::MirrorTangents(Keys, *InCurveModel.GetId(), Min, Max, Midpoint, InCurveEditor);
+	CurveEditor::MirrorKeyPositions(Keys, *InCurveModel.GetId(), Min, Max, Midpoint, InCurveEditor);
 }
 
-template <typename TCallback> requires std::is_invocable_r_v<bool, TCallback, const FLatticeHoverState&, const FLatticeBounds&>
+template <typename TCallback> requires std::is_invocable_r_v<bool, TCallback, const FLatticeHoverState&, const FLatticeBounds&, const FCurveEditor&>
 bool FCurveEditorLatticeTool::ComputeHoverStateForLeftClick(
 	const FGeometry& InViewContainerGeometry, const FPointerEvent& InMouseEvent, TCallback&& InCallback
 	)
@@ -662,21 +859,56 @@ bool FCurveEditorLatticeTool::ComputeHoverStateForLeftClick(
 		const FLatticeHoverState HoverState = ComputeLatticeHoverState(
 			InViewContainerGeometry, InMouseEvent.GetScreenSpacePosition(),
 			ControlPointsInSlate, ComputeDrawControlEdges(ControlPointsInSlate),
-			GetDeformer().NumPointsInWidth(), GetDeformer().NumCells()
+			GetGlobalDeformer().NumPointsInWidth(), GetGlobalDeformer().NumCells()
 			);
-		return InCallback(HoverState, Bounds);
+		return InCallback(HoverState, Bounds, *CurveEditor);
 	}
 
 	return false;
 }
 	
-void FCurveEditorLatticeTool::UnregisterTransactionStateChanged() const
+void FCurveEditorLatticeTool::UnregisterDelegates() const
 {
+	if (const TSharedPtr<FCurveEditor> CurveEditor = WeakCurveEditor.Pin())
+	{
+		CurveEditor->GetTransactionManager()->OnCommandAppended().RemoveAll(this);
+	}
+	
 	if (UTransBuffer* Transactor = GEditor ? Cast<UTransBuffer>(GEditor->Trans) : nullptr)
 	{
 		Transactor->OnTransactionStateChanged().RemoveAll(this);
 	}
 }
+
+void FCurveEditorLatticeTool::SaveUndoState(const FCurveEditor& InCurveEditor)
+{
+	if (ensure(LatticeDeformerState) && GUndo)
+	{
+		UndoObject->Modify();
+		UndoObject->Snapshot = TakeSnapshot(
+			InCurveEditor,
+			LatticeDeformerState->ControlPointToCurveSpace, LatticeDeformerState->GlobalDeformer, LatticeDeformerState->PerCurveData
+			);
+	}
+}
+
+void FCurveEditorLatticeTool::ApplyUndoState(FCurveEditor& InCurveEditor)
+{
+	const bool bWereViewModesCompatible = ApplySnapshot(
+		UndoObject->Snapshot, InCurveEditor,
+		LatticeDeformerState->ControlPointToCurveSpace,
+		LatticeDeformerState->GlobalDeformer, LatticeDeformerState->PerCurveData,
+		// We need to overwrite our serial number, or else in the next frame UpdateDeformerBounds will think the selection has changed and reinit.
+		SelectionSerialNumber
+		);
+
+	// The snapshot could e.g. have been taken in Absolute view mode, but we might currently be in Normalized view.
+	// Since the lattice shapes are incompatible between modes, we'll just regenerate the entire grid.
+	if (!bWereViewModesCompatible)
+	{
+		RecomputeLatticeGrid(InCurveEditor.AsShared());
+	}
+}
 }
 
 #undef LOCTEXT_NAMESPACE
\ No newline at end of file
diff --git a/Engine/Plugins/Editor/CurveEditorTools/Source/CurveEditorTools/Private/Tools/Lattice/CurveEditorLatticeTool.h b/Engine/Plugins/Editor/CurveEditorTools/Source/CurveEditorTools/Private/Tools/Lattice/CurveEditorLatticeTool.h
--- a/Engine/Plugins/Editor/CurveEditorTools/Source/CurveEditorTools/Private/Tools/Lattice/CurveEditorLatticeTool.h
+++ b/Engine/Plugins/Editor/CurveEditorTools/Source/CurveEditorTools/Private/Tools/Lattice/CurveEditorLatticeTool.h
@@ -7,18 +7,21 @@
 #include "ICurveEditorToolExtension.h"
 #include "LatticeControlsDrawData.h"
 #include "LatticeDeformer2D.h"
+#include "LatticeFwd.h"
 #include "Mirror/LatticeEdgeTangentMirrorOp.h"
 #include "Mirror/LatticePointTangentMirrorOp.h"
 #include "Misc/CurveChangeListener.h"
+#include "Misc/CurveViewRebuildListener.h"
 #include "Misc/EnumClassFlags.h"
 #include "Misc/LatticeDrawUtils.h"
+#include "PerCurveLatticeData.h"
 #include "ScopedTransaction.h"
 #include "Templates/SharedPointer.h"
 #include "UObject/UnrealType.h"
 
 class FCurveEditor;
 class FText;
-class UCurveEditorTools_LatticeUndoData;
+class UCurveEditorTools_LatticeUndoObject;
 struct FTransactionContext;
 enum class ETransactionStateEventType : uint8;
 
@@ -35,8 +38,14 @@ enum class ELatticeUpdateFlags : uint8
 ENUM_CLASS_FLAGS(ELatticeUpdateFlags);
 	
 /**
- * Places a simple lattice deformer over the selection.
- * The grid is a simple rectangle. You can reshape by moving the points or edges of the rectangle.
+ * Places a lattice deformer over the selection.
+ * 
+ * The lattice is a quadrilateral you can reshape by moving the points or edges. Reshaping the quad recomputes the underlying key positions.
+ * The lattice does mirroring when:
+ * - Dragging top edge down / bottom edge up: As the edge moves closer to the opposite edge, user set tangents are flattened and then inverted.
+ * - Dragging a control point to the top / bottom edge: Similar as the previous case but a fallout is applied to the key to the left / right of the control point.
+ *
+ * See documentation on FLatticeDeformerState::Deformer and FLatticeDeformerState::PerCurveData to understand the linear spaces the lattice math operates on.
  */
 class FCurveEditorLatticeTool : public ICurveEditorToolExtension, public FGCObject
 {
@@ -74,14 +83,31 @@ private:
 	/** Holds any data that exists only when the deformer exists. The deformer exists the selection consists of least 2 keys at different locations. */
 	struct FLatticeDeformerState
 	{
-		/** Does the actual deforming calculations of the points. */
-		FLatticeDeformer2D Deformer;
+		/**
+		 * Transforms the absolute Deformer.GetControlPoints(), which are just the FKeyPosition values, to curve space values of the current view.
+		 *
+		 * Reminder: CurveSpace are axis values that the SCurveEditorView displays on its axis.
+		 * In Absolute mode, that's the literal FKeyPosition values
+		 * In Normalized mode, Y is in 0 to 1 range (and X continues to be FKeyPosition::InputValue).
+		 */
+		FTransform2d ControlPointToCurveSpace;
 
 		/**
-		 * Tells us when a curve is externally modified, so we can update the bounds.
-		 * Listens for changes on individual curves and cleans up the subscriptions on destruction.
+		 * Empty "deformer" that exists purely for the UI. Its control points, edges, etc. are displayed in the view.
+		 * 
+		 * Each curve has its own lattice deformer that actually moves that curves keys.
+		 * The per curve deformers control points are the result of applying the per-curve transform to Deformer.
+		 * This is needed to support all view modes (Absolute, Normalized, etc.)
+		 *
+		 * The Deformer's control points are always in absolute key space (i.e. exactly the values for FKeyPosition::InputValue and OutputValue).
 		 */
+		FGlobalLatticeDeformer2D GlobalDeformer;
+		TMap<FCurveModelID, FPerCurveLatticeData> PerCurveData; 
+
+		/** Tells us when a curve is externally modified, so we can update the bounds. */
 		FCurveChangeListener CurveChangeListener;
+		/** Tells us when the panel is regenerated, in which case we must regenerate the lattice shape (e.g. could go from absolute -> normalized mode). */
+		FCurveViewRebuildListener PanelRebuildListener;
 		
 		/** Guard to not recompute bounds when FCurveModel::OnCurveModified broadcasts due to a change we have initiated. */
 		bool bIsModifyingCurves = false;
@@ -91,6 +117,9 @@ private:
 		ELatticeUpdateFlags RefreshFlags = ELatticeUpdateFlags::None;
 
 		explicit FLatticeDeformerState(const FLatticeBounds& Lattice, const TSharedRef<FCurveEditor>& InCurveEditor);
+
+		/** @return Control points transformed to SCurveEditorView's curve space. */
+		TArray<FVector2D> TransformControlPointsToCurveSpace() const;
 	};
 	/** Set when there is a valid selection in view. */
 	TOptional<FLatticeDeformerState> LatticeDeformerState;
@@ -103,18 +132,29 @@ private:
 	/** The last mouse position received by OnMouseMove. */
 	FVector2D FinalMouseScreenSpacePosition;
 
+	/** Data used after FDelayedDragData has detected a drag. The user is now actively dragging a control point, edge, or cell. */
 	struct FDragOpData
 	{
 		/** Handles drags that started on a control point. */
 		TUniquePtr<FLatticeDragOp> DragOp;
-		FScopedTransaction Transaction;
+		/** This transaction is active for the duration of the drag. */
+		TUniquePtr<FScopedTransaction> Transaction;
+		
+		/** Adjust tangents while dragging on an edge.*/
+		TMap<FCurveModelID, FLatticeEdgeTangentsMirrorOp> EdgeTangentMirroringOps;
+		/** Adjusts tangents while dragging on a control point. */
+		TMap<FCurveModelID, FLatticePointTangentsMirrorOp> PointTangentMirroringOps;
+
+		/** False until the first drag actually changes keys. We only need to capture undo state just before the first change. */
+		bool bSavedUndoState = false;
 
-		FDragOpData(TUniquePtr<FLatticeDragOp> DragOp, const FText& TransactionText)
+		explicit FDragOpData(TUniquePtr<FLatticeDragOp> DragOp, const FText& TransactionText)
 			: DragOp(MoveTemp(DragOp))
-			, Transaction(TEXT("CurveEditorLatticeTool"), TransactionText, nullptr)
+			, Transaction(MakeUnique<FScopedTransaction>(TEXT("CurveEditorLatticeTool"), TransactionText, nullptr))
 		{}
 	};
-	
+
+	/** Data used while the user is holding down the left mouse button. Once the mouse has moved enough, we "detect" a drag. */
 	struct FDelayedDragData
 	{
 		/** Set when attempting to move a drag handle. This allows us to tell the difference between a click and a click-drag. */
@@ -127,10 +167,6 @@ private:
 
 		/** Set if the drag up was started. */
 		TOptional<FDragOpData> ActiveOperation;
-		/** Adjust tangents while dragging on an edge.*/
-		TOptional<FLatticeEdgeTangentsMirrorOp> EdgeDragTangentMirrorOperation;
-		/** Adjusts tangents while dragging on a control point. */
-		TOptional<FLatticePointTangentsMirrorOp> PointDragTangentMirrorOperation;
 		
 		explicit FDelayedDragData(
 			const FVector2D& InInitialPosition, const FKey& InEffectiveKey, const FLatticeBounds& InBounds, const FLatticeHoverState& InHoverState
@@ -144,22 +180,23 @@ private:
 	 * We store the shape of the lattice grid before starting the drag operation.
 	 * This way, we can restore the shape of the grid when the user undo-es.
 	 */
-	TObjectPtr<UCurveEditorTools_LatticeUndoData> UndoObject;
-	
+	TObjectPtr<UCurveEditorTools_LatticeUndoObject> UndoObject;
+
 	/** Creates or destroys LatticeDeformer based on the selection. */
 	void UpdateDeformerBounds(const TSharedRef<FCurveEditor>& InCurveEditor, ELatticeUpdateFlags InRefreshFlags = ELatticeUpdateFlags::None);
 
 	/** Updates the bounds after the user has externally moved the keys, e.g. MMB move keys, etc. */
 	void OnCurvesModified();
+	/** Regenerates the lattice. */
+	void OnCurvePanelRebuilt() { RecomputeLatticeGrid(ELatticeUpdateFlags::NoTransaction); }
+	/** Handles selection transactional changes: updates the lattice shape change in the same transaction. */
+	void OnUndoCommandAppended(const CurveEditor::FCurveEditorCommandChange& InCommand);
 	/**  */
 	void OnTransactionStateChanged(const FTransactionContext& TransactionContext, ETransactionStateEventType TransactionStateEvent);
 	
 	/** Call this if the keys have externally moved. Tries to fit the lattice grid over the keys or reconstructs a new rectangular one if not possible. */
 	void RecomputeLatticeGrid(const TSharedRef<FCurveEditor>& InCurveEditor, ELatticeUpdateFlags InRefreshFlags = ELatticeUpdateFlags::None);
 	void RecomputeLatticeGrid(ELatticeUpdateFlags InRefreshFlags = ELatticeUpdateFlags::None);
-	
-	/** Stops the active drag operation. */
-	void CancelDrag();
 
 	/** Recomputes all draw data. */
 	void RecomputeDrawData(const FGeometry& InViewContainerGeometry, const TSharedPtr<FCurveEditor>& InCurveEditor);
@@ -167,31 +204,52 @@ private:
 	TArray<FVector2D> ComputeDrawControlPoints(const FCurveEditor& InCurveEditor) const;
 	TArray<FLatticeControlEdge> ComputeDrawControlEdges(TConstArrayView<FVector2D> InControlPoints) const;
 
+	/** Stops the active drag operation. */
+	void CancelDrag();
+	/** Creates a drag-drop operation based on the hover state. */
+	void StartDragDrop(const FGeometry& InViewContainerGeometry, const FPointerEvent& InMouseEvent);
 	/** Starts a new drag drop operation based on whether a control point, edge, or cell is hovered. */
-	void StartDragDrop(const FLatticeBounds& InBounds, const FLatticeHoverState& InHoverState);
+	TOptional<FDragOpData> CreateDragDrop(const FLatticeBounds& InBounds, const FLatticeHoverState& InHoverState);
+	TOptional<FDragOpData> CreateDragDrop_Point(const FLatticeBounds& InBounds, int32 InControlPointIndex);
+	TOptional<FDragOpData> CreateDragDrop_Edge(const FLatticeBounds& InBounds, int32 InEdgeIndex);
+	TOptional<FDragOpData> CreateDragDrop_Cell(const FLatticeBounds& InBounds, int32 InCellIndex);
+
+	// User actions: things the user can directly do in the UI.
+	using FMovePointsSignature = void(FCurveEditorLatticeTool::*)(FPerCurveDeformer2D& InDeformer, FCurveModel& InCurveModel, TConstArrayView<int32> InVertexIndices, TConstArrayView<FVector2D> InNewControlPointPositions, EPropertyChangeType::Type InChangeType);
+	/** Handles dragging a control point, edge, or the lattice center. */
+	void HandleDragPoints(FMovePointsSignature InMovePointsCallback, TConstArrayView<int32> InVertexIndices, TConstArrayView<FVector2D> InNewControlPoints, EPropertyChangeType::Type InChangeType);
+	/** Handles double-clicking an edge. */
+	void HandleMoveEdgeToOpposite(int32 InEdgeIndex);
+	/** Handles double-clicking the lattice center. */
+	void HandleMirrorGrid();
+	
 	/** Called when a control point is dragged. */
-	void OnMoveControlPoint(TConstArrayView<int32> InVertexIndices, TConstArrayView<FVector2D> InNewControlPointPositions, EPropertyChangeType::Type InChangeType);
+	void OnMoveControlPoint(FPerCurveDeformer2D& InDeformer, FCurveModel& InCurveModel, TConstArrayView<int32> InVertexIndices, TConstArrayView<FVector2D> InNewControlPointPositions, EPropertyChangeType::Type InChangeType);
 	/** Called when an edge is dragged. */
-	void OnMoveEdge(TConstArrayView<int32> InVertexIndices, TConstArrayView<FVector2D> InNewControlPointPositions, EPropertyChangeType::Type InChangeType);
+	void OnMoveEdge(FPerCurveDeformer2D& InDeformer, FCurveModel& InCurveModel, TConstArrayView<int32> InVertexIndices, TConstArrayView<FVector2D> InNewControlPointPositions, EPropertyChangeType::Type InChangeType);
 	/** Generic function for updating control points. */
-	void OnSetControlPoints(TConstArrayView<int32> InControlPointsToUpdate, TConstArrayView<FVector2D> InNewControlPoints, EPropertyChangeType::Type InChangeType);
-
+	void OnSetControlPoints(FPerCurveDeformer2D& InDeformer, FCurveModel& InCurveModel, TConstArrayView<int32> InControlPointsToUpdate, TConstArrayView<FVector2D> InNewControlPoints, EPropertyChangeType::Type InChangeType);
+	
 	/** Moves the clicked edge to the opposite edge, effectively flattening the curve sections. */
-	void MoveEdgeToOpposite(int32 InEdgeIndex);
+	void MoveEdgeToOpposite(FPerCurveDeformer2D& InDeformer, FCurveModel& InCurveModel, FCurveEditor& InCurveEditor, int32 InEdgeIndex);
 	/** Mirrors all points in the grid across an imaginary, x-axis aligned edge going through the grid center. Only if the grid consists of a single cell. */
-	void MirrorGrid();
+	void MirrorGrid(FPerCurveDeformer2D& InDeformer, FCurveModel& InCurveModel, FCurveEditor& InCurveEditor);
 	
-	const FLatticeDeformer2D& GetDeformer() const { check(LatticeDeformerState); return LatticeDeformerState->Deformer; }
-	FLatticeDeformer2D& GetDeformer() { check(LatticeDeformerState); return LatticeDeformerState->Deformer; }
+	const FLatticeDeformer2D& GetGlobalDeformer() const { return LatticeDeformerState->GlobalDeformer; }
+	FLatticeDeformer2D& GetGlobalDeformer() { return LatticeDeformerState->GlobalDeformer; }
 
 	/**
 	 * Invokes InCallback if there is a lattice grid shown and the user left clicks somewhere.
 	 * @return False if InCallback was not invoked. Otherwise, whatever InCallback returned. */
-	template<typename TCallback> requires std::is_invocable_r_v<bool, TCallback, const FLatticeHoverState&, const FLatticeBounds&>
+	template<typename TCallback> requires std::is_invocable_r_v<bool, TCallback, const FLatticeHoverState&, const FLatticeBounds&, const FCurveEditor&>
 	bool ComputeHoverStateForLeftClick(const FGeometry& InViewContainerGeometry, const FPointerEvent& InMouseEvent, TCallback&& InCallback);
 
-	/** Shared logic for unregistering global transaction delegate. */
-	void UnregisterTransactionStateChanged() const;
+	/** Shared logic for unregistering global delegates. */
+	void UnregisterDelegates() const;
+
+	/** Saves the undo state in the undo object. */
+	void SaveUndoState(const FCurveEditor& InCurveEditor);
+	void ApplyUndoState(FCurveEditor& InCurveEditor);
 };
 
 inline void FCurveEditorLatticeTool::RecomputeLatticeGrid(ELatticeUpdateFlags InRefreshFlags)
diff --git a/Engine/Plugins/Editor/CurveEditorTools/Source/CurveEditorTools/Private/Tools/Lattice/Drag/LatticeDragOp_MoveControlPoints.cpp b/Engine/Plugins/Editor/CurveEditorTools/Source/CurveEditorTools/Private/Tools/Lattice/Drag/LatticeDragOp_MoveControlPoints.cpp
--- a/Engine/Plugins/Editor/CurveEditorTools/Source/CurveEditorTools/Private/Tools/Lattice/Drag/LatticeDragOp_MoveControlPoints.cpp
+++ b/Engine/Plugins/Editor/CurveEditorTools/Source/CurveEditorTools/Private/Tools/Lattice/Drag/LatticeDragOp_MoveControlPoints.cpp
@@ -20,7 +20,10 @@ FLatticeDragOp_MoveControlPoints::FLatticeDragOp_MoveControlPoints(
 	, OnInteractiveDragControlPointsDelegate(MoveTemp(InInteractiveDragControlPointsDelegate))
 	, OnFinishDragControlPointsDelegate(MoveTemp(InFinishDragControlPointsDelegate))
 	, UpdatedControlPointPositions(InitialControlPoints)
-{}
+{
+	UE_LOG(LogTemp, Warning, TEXT("SlateMin: %s \t SlateMax:%s"), *Bounds.MinSlatePosition.ToString(), *Bounds.MaxSlatePosition.ToString());
+	UE_LOG(LogTemp, Warning, TEXT("Min: %s \t Max:%s"), *Bounds.MinValues.ToString(), *Bounds.MaxValues.ToString());
+}
 	
 void FLatticeDragOp_MoveControlPoints::OnMoveMouse(const FGeometry& InGeometry, const FVector2D& InScreenPosition)
 {
diff --git a/Engine/Plugins/Editor/CurveEditorTools/Source/CurveEditorTools/Private/Tools/Lattice/LatticeDeformer2D.cpp b/Engine/Plugins/Editor/CurveEditorTools/Source/CurveEditorTools/Private/Tools/Lattice/LatticeDeformer2D.cpp
--- a/Engine/Plugins/Editor/CurveEditorTools/Source/CurveEditorTools/Private/Tools/Lattice/LatticeDeformer2D.cpp
+++ b/Engine/Plugins/Editor/CurveEditorTools/Source/CurveEditorTools/Private/Tools/Lattice/LatticeDeformer2D.cpp
@@ -34,57 +34,18 @@ namespace UE::CurveEditorTools
 {
 FLatticeDeformer2D::FLatticeDeformer2D(int32 InNumCellsInX, int32 InNumCellsInY, const FVector2D& InBottomLeft, const FVector2D& InTopRight)
 	: CellDimensions(InNumCellsInX, InNumCellsInY)
+	, ControlPoints(GenerateControlPoints(InNumCellsInX, InNumCellsInY, InBottomLeft, InTopRight))
 {
-	// Lattice with 0 width or height does not make sense. We'd divide by 0 in ComputeCoordinatesInCell.
-	if (!ensure(InBottomLeft.X < InTopRight.X) || !ensure(InBottomLeft.Y < InTopRight.Y))
+	if (!ControlPoints.IsEmpty())
 	{
-		// User specified invalid bounds. The lattice will be empty.
-		return;
+		Cells.SetNum(NumCells());
 	}
-	ControlPoints.SetNumUninitialized(NumPointsInWidth() * NumPointsInHeight());
-
-	const double DeltaX = (InTopRight.X - InBottomLeft.X) / static_cast<double>(InNumCellsInX);
-	const double DeltaY = (InTopRight.Y - InBottomLeft.Y) / static_cast<double>(InNumCellsInY);
-	for (int32 X = 0; X < NumPointsInWidth(); ++X)
-	{
-		for (int32 Y = 0; Y < NumPointsInWidth(); ++Y)
-		{
-			const int32 Index = MatrixIndicesToFlatIndex(X, Y, NumPointsInWidth());
-			const FVector2D ControlPosition(InBottomLeft.X + DeltaX * X, InBottomLeft.Y + DeltaY * Y);
-			ControlPoints[Index] = ControlPosition;
-		}
-	}
-
-	Cells.SetNum(NumCells());
 }
 
-void FLatticeDeformer2D::AddPoints_BeforeLatticeMoved(TConstArrayView<FKeyInCurveId> InKeys, TConstArrayView<FVector2D> InPoints)
-{
-	if (!ensure(InKeys.Num() == InPoints.Num()))
-	{
-		return;
-	}
-
-	// No need to make this ParallelFor: 4500 keys take about 0.6ms.
-	for (int32 KeyIndex = 0; KeyIndex < InKeys.Num(); ++KeyIndex)
-	{
-		const FVector2D& KeyPoint = InPoints[KeyIndex];
-		// Check whether in cell is relatively cheap (addition & multiplication) ...
-		const int32 CellIndex = FindCellContainingPoint(KeyPoint);
-		if (!Cells.IsValidIndex(CellIndex))
-		{
-			continue;
-		}
-		
-		TArray<FLatticeKeyCoords>& Keys = Cells[CellIndex].KeyCoords;
-		TArray<FKeyInCurveId>& KeyIds = Cells[CellIndex].KeyIds;
-		KeyIds.Emplace(InKeys[KeyIndex]);
-		Keys.Emplace(ComputeSingleCellWeights(GetCellIndices(CellIndex), InPoints[KeyIndex]));
-	}
-}
-
-void FLatticeDeformer2D::UpdateControlPoints(TConstArrayView<int32> InControlPointsToUpdate, TConstArrayView<FVector2D> InNewControlPoints,
-	const TFunctionRef<void(const FKeyInCurveId&, const FVector2D&)> HandleKeyChange)
+void FLatticeDeformer2D::UpdateControlPoints(
+	TConstArrayView<int32> InControlPointsToUpdate, TConstArrayView<FVector2D> InNewControlPoints,
+	const TFunctionRef<void(const FPointIndex&, const FVector2D&)> HandleKeyChange
+	)
 {
 	const bool bValidIndices = Algo::AllOf(InControlPointsToUpdate, [this](int32 InIndex){ return ControlPoints.IsValidIndex(InIndex); });
 	if (!ensure(bValidIndices && InControlPointsToUpdate.Num() == InNewControlPoints.Num())
@@ -165,19 +126,6 @@ FCellVertexIndices FLatticeDeformer2D::GetCellIndices(int32 InCellIndex) const
 	return ensure(bIsInRange) ? CurveEditorTools::GetCellIndices(InCellIndex, NumPointsInWidth()) : FCellVertexIndices();
 }
 
-void FLatticeDeformer2D::SetKeyIdsInCell(int32 InCellIndex, TArray<FKeyInCurveId> InNewKeyCoords)
-{
-	if (ensure(Cells.IsValidIndex(InCellIndex)))
-	{
-		Cells[InCellIndex].KeyIds = MoveTemp(InNewKeyCoords);
-	}
-}
-
-TConstArrayView<FKeyInCurveId> FLatticeDeformer2D::GetKeyIdsInCell(int32 InCellIndex) const
-{
-	return ensure(Cells.IsValidIndex(InCellIndex)) ? Cells[InCellIndex].KeyIds : TConstArrayView<FKeyInCurveId>();
-} 
-
 void FLatticeDeformer2D::SetKeyCoordsInCell(int32 InCellIndex, TArray<FLatticeKeyCoords> InNewKeyCoords)
 {
 	if (ensure(Cells.IsValidIndex(InCellIndex)))
@@ -215,6 +163,16 @@ int32 FLatticeDeformer2D::FindCellContainingPoint(const FVector2D& InPoint) cons
 	}
 }
 
+int32 FLatticeDeformer2D::NumKeys() const
+{
+	int32 Count = 0;
+	for (int32 CellIndex = 0; CellIndex < Cells.Num(); ++CellIndex)
+	{
+		Count += Cells[CellIndex].KeyCoords.Num();
+	}
+	return Count;
+}
+
 int32 FLatticeDeformer2D::NumEdges() const
 {
 	if (NumPointsInWidth() == 0 || NumPointsInHeight() == 0)
@@ -233,8 +191,40 @@ int32 FLatticeDeformer2D::NumCells() const
 {
 	return (NumPointsInWidth() - 1) * (NumPointsInHeight() - 1);
 }
+	
+void FLatticeDeformer2D::AddPoints_BeforeLatticeMoved(
+	TConstArrayView<FVector2D> InPoints,
+	TFunctionRef<void(int32 InputIndex, const FPointIndex&)> OnPointIndexed
+	)
+{
+	// No need to make this ParallelFor: 4500 keys take about 0.6ms.
+	for (int32 PointIndex = 0; PointIndex < InPoints.Num(); ++PointIndex)
+	{
+		const FVector2D& KeyPoint = InPoints[PointIndex];
+		// Checking whether the point is in the cell is relatively cheap (addition & multiplication) ...
+		const int32 CellIndex = FindCellContainingPoint(KeyPoint);
+		if (!Cells.IsValidIndex(CellIndex))
+		{
+			continue;
+		}
+		
+		TArray<FLatticeKeyCoords>& Keys = Cells[CellIndex].KeyCoords;
+		const int32 ItemIndex = Keys.Emplace(ComputeSingleCellWeights(GetCellIndices(CellIndex), InPoints[PointIndex]));
+		OnPointIndexed(PointIndex, FPointIndex(CellIndex, ItemIndex));
+	}
+}
 
-void FLatticeDeformer2D::RecomputeCell(int32 InCellIndex, const TFunctionRef<void(const FKeyInCurveId&, const FVector2D&)> HandleKeyChange)
+FLatticeKeyCoords FLatticeDeformer2D::ComputeSingleCellWeights(const FCellVertexIndices& InIndices, const FVector2D& InPoint) const
+{
+	const auto[TopLeft, TopRight, BottomRight, BottomLeft] = InIndices;
+	const FVector2D Vertex_TopRight = ControlPoints[TopRight];
+	const FVector2D Vertex_BottomLeft = ControlPoints[BottomLeft];
+	return Math::ComputeCoordinatesInCell(InPoint, Vertex_BottomLeft, Vertex_TopRight);
+}
+
+void FLatticeDeformer2D::RecomputeCell(
+	int32 InCellIndex, const TFunctionRef<void(const FPointIndex&, const FVector2D&)> HandleKeyChange
+	)
 {
 	if (!Cells.IsValidIndex(InCellIndex))
 	{
@@ -253,10 +243,9 @@ void FLatticeDeformer2D::RecomputeCell(int32 InCellIndex, const TFunctionRef<voi
 		Positions[KeyIndex] = RecomputeKeyValue(KeyData, CellVertices);
 	});
 
-	TArray<FKeyInCurveId>& KeyIds = Cells[InCellIndex].KeyIds;
 	for (int32 Index = 0; Index < Keys.Num(); ++Index)
 	{
-		HandleKeyChange(KeyIds[Index], Positions[Index]);
+		HandleKeyChange(FPointIndex(InCellIndex, Index), Positions[Index]);
 	}
 }
 
@@ -274,12 +263,4 @@ FVector2D FLatticeDeformer2D::RecomputeKeyValue(const FLatticeKeyCoords& PointCo
 	const FVector2D FinalPoint = BottomLeft + TopLeft + BottomRight + TopRight;
 	return { FinalPoint.X, FinalPoint.Y };
 }
-
-FLatticeKeyCoords FLatticeDeformer2D::ComputeSingleCellWeights(const FCellVertexIndices& InIndices, const FVector2D& InPoint) const
-{
-	const auto[TopLeft, TopRight, BottomRight, BottomLeft] = InIndices;
-	const FVector2D Vertex_TopRight = ControlPoints[TopRight];
-	const FVector2D Vertex_BottomLeft = ControlPoints[BottomLeft];
-	return Math::ComputeCoordinatesInCell(InPoint, Vertex_BottomLeft, Vertex_TopRight);
-}
 }
\ No newline at end of file
diff --git a/Engine/Plugins/Editor/CurveEditorTools/Source/CurveEditorTools/Private/Tools/Lattice/LatticeDeformer2D.h b/Engine/Plugins/Editor/CurveEditorTools/Source/CurveEditorTools/Private/Tools/Lattice/LatticeDeformer2D.h
--- a/Engine/Plugins/Editor/CurveEditorTools/Source/CurveEditorTools/Private/Tools/Lattice/LatticeDeformer2D.h
+++ b/Engine/Plugins/Editor/CurveEditorTools/Source/CurveEditorTools/Private/Tools/Lattice/LatticeDeformer2D.h
@@ -7,7 +7,6 @@
 #include "HAL/Platform.h"
 #include "Math/MathFwd.h"
 #include "Math/Vector2D.h"
-#include "Misc/KeyInCurveId.h"
 #include "Templates/FunctionFwd.h"
 
 namespace UE::CurveEditorTools
@@ -33,6 +32,17 @@ struct FLatticeControlEdge
 	{}
 };
 
+/** Describes how to locate a key in FLatticeDeformer2D */
+struct FPointIndex
+{
+	/** Index of the cell in the lattice. */
+	int32 CellIndex;
+	/** Index in the cell's key array. */
+	int32 IndexInCell;
+
+	explicit FPointIndex(int32 InCellIndex, int32 InIndexInCell) : CellIndex(InCellIndex), IndexInCell(InIndexInCell) {}
+};
+
 /** Identifies the normalized coordinates a key has in a cell. */
 using FLatticeKeyCoords = FVector2D;
 	
@@ -69,18 +79,11 @@ public:
 
 	/** Completely resets this deformer and sets a new grid.  */
 	explicit FLatticeDeformer2D(int32 InNumCellsInX, int32 InNumCellsInY, const FVector2D& InBottomLeft, const FVector2D& InTopRight);
-
-	/**
-	 * Computes the weights for each of the points.
-	 * 
-	 * This function assumes that none of the control points have yet been moved: if so, the computed weights are incorrect.
-	 * The computation of the weights requires that all cells are rectangular.
-	 */
-	void AddPoints_BeforeLatticeMoved(TConstArrayView<FKeyInCurveId> InKeys, TConstArrayView<FVector2D> InPoints);
 	
 	/** Moves the specified control points and recomputes the affected points. */
-	void UpdateControlPoints(TConstArrayView<int32> InControlPointsToUpdate, TConstArrayView<FVector2D> InNewControlPoints,
-		const TFunctionRef<void(const FKeyInCurveId&, const FVector2D&)> HandleKeyChange
+	void UpdateControlPoints(
+		TConstArrayView<int32> InControlPointsToUpdate, TConstArrayView<FVector2D> InNewControlPoints,
+		const TFunctionRef<void(const FPointIndex&, const FVector2D&)> HandleKeyChange
 		);
 
 	/** Sets the value of the control points without recomputing the values of the keys. Useful if you are resetting to a previous state. */
@@ -102,11 +105,6 @@ public:
 	 */
 	FCellVertexIndices GetCellIndices(int32 InCellIndex) const;
 
-	/** Sets the key IDs of the keys in the cell. */
-	void SetKeyIdsInCell(int32 InCellIndex, TArray<FKeyInCurveId> InNewKeyCoords);
-	/** @return The key IDs of the keys in the cell. */
-	TConstArrayView<FKeyInCurveId> GetKeyIdsInCell(int32 InCellIndex) const;
-
 	/** Sets the key coords of keys in the cell. This is useful for restoring previous data. */
 	void SetKeyCoordsInCell(int32 InCellIndex, TArray<FLatticeKeyCoords> InNewKeyCoords);
 	/** @return The coordinates of the keys in the cell. */
@@ -114,6 +112,9 @@ public:
 	
 	/** @return The vertices of the cell containing InPoint, or all INDEX_NONE if not contained.  */
 	int32 FindCellContainingPoint(const FVector2D& InPoint) const;
+
+	/** @return Counts the number of keys stored in all cells. */
+	int32 NumKeys() const;
 	
 	int32 NumControlPoints() const { return ControlPoints.Num(); }
 	int32 NumPointsInWidth() const { return CellDimensions.X + 1; }
@@ -122,6 +123,25 @@ public:
 	int32 NumCellsInWidth() const { return CellDimensions.X; }
 	int32 NumCellsInHeight() const { return CellDimensions.Y; }
 	int32 NumCells() const;
+
+protected:
+	
+	/**
+	 * Adds the given points to the lattice grid.
+	 * 
+	 * This function assumes that none of the control points have yet been moved: if so, the computed weights are incorrect.
+	 * The computation of the weights requires that all cells are rectangular.
+	 *
+	 * Exposed only to subclasses so subclasses can control how points are added.
+	 * If this was public, you could downcast sub-classes to this base class and break the Liskov-Substitution principle.
+	 */
+	void AddPoints_BeforeLatticeMoved(
+		TConstArrayView<FVector2D> InPoints,
+		TFunctionRef<void(int32 InputIndex, const FPointIndex&)> OnPointIndexed = [](auto, auto){}
+		);
+	
+	/** Computes the weights for InKeyData such that only the corners of the cell it is in affect it. */
+	FLatticeKeyCoords ComputeSingleCellWeights(const FCellVertexIndices& InIndices, const FVector2D& InPoint) const;
 	
 private:
 	
@@ -138,7 +158,6 @@ private:
 
 	struct FCellData
 	{
-		TArray<FKeyInCurveId> KeyIds;
 		TArray<FLatticeKeyCoords> KeyCoords;
 	};
 	/**
@@ -150,12 +169,86 @@ private:
 	TArray<FCellData> Cells;
 
 	/** Recomputes all key positions in the cell. Call this when cell has changed its control points. */
-	void RecomputeCell(int32 InCellIndex, const TFunctionRef<void(const FKeyInCurveId&, const FVector2D&)> HandleKeyChange);
+	void RecomputeCell(int32 InCellIndex, const TFunctionRef<void(const FPointIndex&, const FVector2D&)> HandleKeyChange);
 
 	/** Computes the value the key should by computing the weighted linear product of all control points affecting it. */
 	FVector2D RecomputeKeyValue(const FLatticeKeyCoords& InPointCoords, const FCellVertexIndices& InIndices) const;
-	/** Computes the weights for InKeyData such that only the corners of the cell it is in affect it. */
-	FLatticeKeyCoords ComputeSingleCellWeights(const FCellVertexIndices& InIndices, const FVector2D& InPoint) const;
+};
+
+/** Allows you to attach meta-data to each key position. */
+template<typename TPointMetaData>
+class TLatticeDeformer2D : public FLatticeDeformer2D
+{
+	using Super = FLatticeDeformer2D;
+	
+public:
+
+	/** Completely resets this deformer and sets a new grid.  */
+	explicit TLatticeDeformer2D(int32 InNumCellsInX, int32 InNumCellsInY, const FVector2D& InBottomLeft, const FVector2D& InTopRight)
+		: Super(InNumCellsInX, InNumCellsInY, InBottomLeft, InTopRight)
+	{
+		CellMetaData.SetNum(NumCells());
+	}
+
+	/**
+	 * Computes the weights for each of the points.
+	 * 
+	 * This function assumes that none of the control points have yet been moved: if so, the computed weights are incorrect.
+	 * The computation of the weights requires that all cells are rectangular.
+	 */
+	void AddPoints_BeforeLatticeMoved(TConstArrayView<TPointMetaData> InKeys, TConstArrayView<FVector2D> InPoints)
+	{
+		if (!ensure(InKeys.Num() == InPoints.Num()))
+		{
+			return;
+		}
+		
+		Super::AddPoints_BeforeLatticeMoved(InPoints, [this, &InKeys](int32 InputIndex, const FPointIndex& PointIndex)
+		{
+			TArray<TPointMetaData>& MetaData = CellMetaData[PointIndex.CellIndex].PointMetaData;
+			MetaData.SetNum(PointIndex.IndexInCell + 1); // Handles the case that API user called Super::AddPoints_BeforeLatticeMoved 
+			MetaData[PointIndex.IndexInCell] = InKeys[InputIndex];
+		});
+	}
+	
+	/** Moves the specified control points and recomputes the affected points. */
+	void UpdateControlPoints(TConstArrayView<int32> InControlPointsToUpdate, TConstArrayView<FVector2D> InNewControlPoints,
+		const TFunctionRef<void(const TPointMetaData&, const FVector2D&)> HandleKeyChange
+		)
+	{
+		Super::UpdateControlPoints(InControlPointsToUpdate, InNewControlPoints,
+			[this, &HandleKeyChange](const FPointIndex& PointIndex, const FVector2D& UpdatedPosition)
+			{
+				const TPointMetaData& MetaData = CellMetaData[PointIndex.CellIndex].PointMetaData[PointIndex.IndexInCell];
+				HandleKeyChange(MetaData, UpdatedPosition);
+			});
+	}
+	
+	/** @return The key IDs of the keys in the cell. */
+	TConstArrayView<TPointMetaData> GetCellMetaData(int32 InCellIndex) const
+	{
+		return ensure(CellMetaData.IsValidIndex(InCellIndex)) ? CellMetaData[InCellIndex].PointMetaData : TConstArrayView<TPointMetaData>();
+	}
+	
+	/** Sets the key IDs of the keys in the cell. */
+	void SetCellMetaData(int32 InCellIndex, TArray<TPointMetaData> InNewKeyCoords)
+	{
+		if (ensure(CellMetaData.IsValidIndex(InCellIndex)))
+		{
+			CellMetaData[InCellIndex].PointMetaData = MoveTemp(InNewKeyCoords);
+		}
+	}
+
+private:
+
+	struct FCellMetaData
+	{
+		/** Same length as the equivalent KeyCoords array of FCellData. */
+		TArray<TPointMetaData> PointMetaData;
+	};
+
+	/** Same length as FLatticeDeformer2D::Cells */
+	TArray<FCellMetaData> CellMetaData;
 };
 }
 
diff --git a/Engine/Plugins/Editor/CurveEditorTools/Source/CurveEditorTools/Private/Tools/Lattice/LatticeFwd.h b/Engine/Plugins/Editor/CurveEditorTools/Source/CurveEditorTools/Private/Tools/Lattice/LatticeFwd.h
new file mode 100644
--- /dev/null
+++ b/Engine/Plugins/Editor/CurveEditorTools/Source/CurveEditorTools/Private/Tools/Lattice/LatticeFwd.h
@@ -0,0 +1,15 @@
+﻿// Copyright Epic Games, Inc. All Rights Reserved.
+
+#pragma once
+
+struct FKeyHandle;
+
+namespace UE::CurveEditorTools
+{
+class FLatticeDeformer2D;
+template<typename TPointMetaData> class TLatticeDeformer2D;
+
+/** We never add any points to this deformer. It's just used to keep track of the control points & edges displayed in the view. */
+using FGlobalLatticeDeformer2D = FLatticeDeformer2D;
+using FPerCurveDeformer2D = TLatticeDeformer2D<FKeyHandle>;
+}
diff --git a/Engine/Plugins/Editor/CurveEditorTools/Source/CurveEditorTools/Private/Tools/Lattice/Mirror/LatticeEdgeTangentMirrorOp.cpp b/Engine/Plugins/Editor/CurveEditorTools/Source/CurveEditorTools/Private/Tools/Lattice/Mirror/LatticeEdgeTangentMirrorOp.cpp
--- a/Engine/Plugins/Editor/CurveEditorTools/Source/CurveEditorTools/Private/Tools/Lattice/Mirror/LatticeEdgeTangentMirrorOp.cpp
+++ b/Engine/Plugins/Editor/CurveEditorTools/Source/CurveEditorTools/Private/Tools/Lattice/Mirror/LatticeEdgeTangentMirrorOp.cpp
@@ -12,15 +12,30 @@ namespace UE::CurveEditorTools
 {
 namespace MirrorEdgeDetail
 {
-static bool CanUseEdgeForMirroring(int32 InEdgeIndex, const FLatticeDeformer2D& InDeformer)
+static bool CanMirror(int32 InEdgeIndex, const FLatticeDeformer2D& InDeformer)
 {
 	const bool bCorrectEdge = InEdgeIndex == GetEdgeIndexInSingleCellGrid(ELatticeEdgeType::Top)
 		|| InEdgeIndex == GetEdgeIndexInSingleCellGrid(ELatticeEdgeType::Bottom);
 	return bCorrectEdge && InDeformer.NumCells() == 1;
 }
 	
-static void GetDragStartAndMidpointHeights(
-	int32 InEdgeIndex, const FLatticeDeformer2D& InDeformer, EMirroredEdgeType& OutEdgeType, double& OutDragStartHeight, double& OutMidPointHeight
+static TOptional<EMirroredEdgeType> GetMirroredEdgeType(int32 InEdgeIndex)
+{
+	constexpr int32 TopIndex = GetEdgeIndexInSingleCellGrid(ELatticeEdgeType::Top);
+	constexpr int32 BottomIndex = GetEdgeIndexInSingleCellGrid(ELatticeEdgeType::Bottom);
+	if (TopIndex == InEdgeIndex)
+	{
+		return EMirroredEdgeType::Top;
+	}
+	if (BottomIndex == InEdgeIndex)
+	{
+		return EMirroredEdgeType::Bottom;
+	}
+	return {};
+}
+	
+static bool GetDragStartAndMidpointHeights(
+	int32 InEdgeIndex, const FLatticeDeformer2D& InDeformer, double& OutDragStartHeight, double& OutMidPointHeight
 )
 {
 	constexpr int32 TopIndex = GetEdgeIndexInSingleCellGrid(ELatticeEdgeType::Top);
@@ -31,42 +46,68 @@ static void GetDragStartAndMidpointHeights(
 	const auto[OppositeStartVertIndex, OppositeEndVertIndex] = InDeformer.GetEdgeIndices(OppositeEdgeIndex);
 
 	const TConstArrayView<FVector2D> ControlPoints = InDeformer.GetControlPoints();
-	OutEdgeType = bIsTopEdge ? EMirroredEdgeType::Top : EMirroredEdgeType::Bottom;
-	OutDragStartHeight = ChooseConsistentEdgeVert(OutEdgeType, ControlPoints[StartVertIndex], ControlPoints[EndVertIndex]);
-	OutMidPointHeight = ChooseConsistentEdgeVert(OutEdgeType, ControlPoints[OppositeStartVertIndex], ControlPoints[OppositeEndVertIndex]);
+	if (const TOptional<EMirroredEdgeType> OutEdgeType = GetMirroredEdgeType(InEdgeIndex))
+	{
+		OutDragStartHeight = ChooseConsistentEdgeVert(*OutEdgeType, ControlPoints[StartVertIndex], ControlPoints[EndVertIndex]);
+		OutMidPointHeight = ChooseConsistentEdgeVert(*OutEdgeType, ControlPoints[OppositeStartVertIndex], ControlPoints[OppositeEndVertIndex]);
+		return true;
+	}
+	return false;
 }
 }
 
-TOptional<FLatticeEdgeTangentsMirrorOp> FLatticeEdgeTangentsMirrorOp::MakeFromDraggedEdge(
-	int32 InEdgeIndex, const FLatticeDeformer2D& InDeformer, const FCurveEditor& InCurveEditor
+TOptional<FLatticeEdgeTangentsMirrorOp> FLatticeEdgeTangentsMirrorOp::MakeMirrorOpForDragLatticeEdge(
+	int32 InEdgeIndex, const FCurveModelID& InCurveId, const TLatticeDeformer2D<FKeyHandle>& InDeformer, const FCurveEditor& InCurveEditor
 	)
 {
-	if (MirrorEdgeDetail::CanUseEdgeForMirroring(InEdgeIndex, InDeformer))
+	if (!MirrorEdgeDetail::CanMirror(InEdgeIndex, InDeformer))
 	{
-		EMirroredEdgeType EdgeType;
-		double DragStartHeight, MidPointHeight;
-		MirrorEdgeDetail::GetDragStartAndMidpointHeights(InEdgeIndex, InDeformer, EdgeType, DragStartHeight, MidPointHeight);
-
-		CurveEditor::FTangentMirrorSolver TangentSolver(DragStartHeight, MidPointHeight);
-		bool bAtLeastOne = false;
-		for (const TPair<FCurveModelID, TArray<FKeyHandle>>& Pair : SplitKeysInCurves(InDeformer.GetKeyIdsInCell(0)))
-		{
-			bAtLeastOne |= TangentSolver.AddTangents(InCurveEditor, Pair.Key, Pair.Value);
-		}
+		return {};
+	}
+	
+	TOptional<FCurveMirrorData> MirrorData = ComputeTangentMirrorData(InEdgeIndex, InCurveEditor, InDeformer, InCurveId, InDeformer.GetCellMetaData(0));
+	if (!MirrorData)
+	{
+		return {};
+	}
+	
+	return FLatticeEdgeTangentsMirrorOp(InEdgeIndex, MoveTemp(*MirrorData));
+}
 
-		if (bAtLeastOne)
-		{
-			return FLatticeEdgeTangentsMirrorOp(InCurveEditor, EdgeType, MoveTemp(TangentSolver));
-		}
+void FLatticeEdgeTangentsMirrorOp::OnMoveEdge(TConstArrayView<FVector2D> InNewEdge, const FCurveEditor& InCurveEditor)
+{
+	const TOptional<EMirroredEdgeType> EdgeType = MirrorEdgeDetail::GetMirroredEdgeType(EdgeIndex);
+	if (!ensure(EdgeType.IsSet()) || !ensureMsgf(InNewEdge.Num() == 2, TEXT("We expected an edge")))
+	{
+		return;
 	}
-	return {};
+	
+	FCurveModel* CurveModel = InCurveEditor.FindCurve(CurveMirrorData.CurveId);
+	if (!CurveModel)
+	{
+		return;
+	}
+	CurveModel->Modify();
+		
+	CurveMirrorData.Solver.OnMoveEdge(InCurveEditor, ChooseConsistentEdgeVert(*EdgeType, InNewEdge[0], InNewEdge[1]));
 }
 
-void FLatticeEdgeTangentsMirrorOp::OnMoveEdge(TConstArrayView<FVector2D> InNewEdge)
+TOptional<FLatticeEdgeTangentsMirrorOp::FCurveMirrorData> FLatticeEdgeTangentsMirrorOp::ComputeTangentMirrorData(
+	int32 InEdgeIndex, const FCurveEditor& InCurveEditor,
+	const FLatticeDeformer2D& InCurveDeformer, const FCurveModelID& InCurve, TConstArrayView<FKeyHandle> InKeys
+	)
 {
-	if (ensureMsgf(InNewEdge.Num() == 2, TEXT("We expected an edge")))
+	double DragStartHeight, MidPointHeight;
+	if (!MirrorEdgeDetail::GetDragStartAndMidpointHeights(InEdgeIndex, InCurveDeformer, DragStartHeight, MidPointHeight))
 	{
-		TangentMirrorSolver.OnMoveEdge(CurveEditor,  ChooseConsistentEdgeVert(EdgeType, InNewEdge[0], InNewEdge[1]));
+		return {};
 	}
+	
+	CurveEditor::FTangentMirrorSolver TangentSolver(DragStartHeight, MidPointHeight);
+	if (TangentSolver.AddTangents(InCurveEditor, InCurve, InKeys))
+	{
+		return FCurveMirrorData(InCurve, MoveTemp(TangentSolver));
+	}
+	return {};
+}
 }
-}
\ No newline at end of file
diff --git a/Engine/Plugins/Editor/CurveEditorTools/Source/CurveEditorTools/Private/Tools/Lattice/Mirror/LatticeEdgeTangentMirrorOp.h b/Engine/Plugins/Editor/CurveEditorTools/Source/CurveEditorTools/Private/Tools/Lattice/Mirror/LatticeEdgeTangentMirrorOp.h
--- a/Engine/Plugins/Editor/CurveEditorTools/Source/CurveEditorTools/Private/Tools/Lattice/Mirror/LatticeEdgeTangentMirrorOp.h
+++ b/Engine/Plugins/Editor/CurveEditorTools/Source/CurveEditorTools/Private/Tools/Lattice/Mirror/LatticeEdgeTangentMirrorOp.h
@@ -24,28 +24,43 @@ class FLatticeEdgeTangentsMirrorOp
 {
 public:
 
-	/** @return Constructed instance that will mirror all keys in InDeformer. Returns valid if the edge is a top or bottom edge in a single celled grid. */
-	static TOptional<FLatticeEdgeTangentsMirrorOp> MakeFromDraggedEdge(
-		int32 InEdgeIndex, const FLatticeDeformer2D& InDeformer, const FCurveEditor& InCurveEditor UE_LIFETIMEBOUND
-		);
-
+	/**
+	 * Util function that creates a mirror op for dragging InEdgeIndex.
+	 * @return Valid if InEdgeIndex is an edge that should do mirroring when dragged, i.e. a top or bottom edge
+	 */
+	static TOptional<FLatticeEdgeTangentsMirrorOp> MakeMirrorOpForDragLatticeEdge(
+		int32 InEdgeIndex, const FCurveModelID& InCurveId, const TLatticeDeformer2D<FKeyHandle>& InDeformer,
+		const FCurveEditor& InCurveEditor UE_LIFETIMEBOUND
+	);
+	
 	/** Call with info about where the new edge is located. */
-	void OnMoveEdge(TConstArrayView<FVector2D> InNewEdge);
+	void OnMoveEdge(TConstArrayView<FVector2D> InNewEdge, const FCurveEditor& InCurveEditor);
 	
 private:
 
-	/** Used to get and set key attributes */
-	const FCurveEditor& CurveEditor;
-	/** Edge type of the dragged edge. Important for determining height of sloped edges. */
-	const EMirroredEdgeType EdgeType;
-	
+	/** Indexs of the edge being moved. */
+	const int32 EdgeIndex;
+
+	struct FCurveMirrorData
+	{
+		FCurveModelID CurveId;
+		CurveEditor::FTangentMirrorSolver Solver;
+
+		explicit FCurveMirrorData(FCurveModelID InCurveId, CurveEditor::FTangentMirrorSolver&& InSolver)
+			: CurveId(InCurveId), Solver(MoveTemp(InSolver))
+		{}
+	};
 	/** Computes the tangents */
-	CurveEditor::FTangentMirrorSolver TangentMirrorSolver;
+	FCurveMirrorData CurveMirrorData;
 	
-	FLatticeEdgeTangentsMirrorOp(const FCurveEditor& InCurveEditor, const EMirroredEdgeType InEdgeType, CurveEditor::FTangentMirrorSolver InSolver)
-		: CurveEditor(InCurveEditor)
-		, EdgeType(InEdgeType)
-		, TangentMirrorSolver(MoveTemp(InSolver))
+	explicit FLatticeEdgeTangentsMirrorOp(int32 InEdgeIndex, FCurveMirrorData InMirrorData)
+		: EdgeIndex(InEdgeIndex), CurveMirrorData(MoveTemp(InMirrorData))
 	{}
+	
+	/** Tracks keys that have tangents that can be mirrored. The mirror midpoints are computed based off of InCurveDeformer. */
+	static TOptional<FCurveMirrorData> ComputeTangentMirrorData(
+		int32 InEdgeIndex, const FCurveEditor& InCurveEditor,
+		const FLatticeDeformer2D& InCurveDeformer, const FCurveModelID& InCurve, TConstArrayView<FKeyHandle> InKeys
+		);
 };
 }
diff --git a/Engine/Plugins/Editor/CurveEditorTools/Source/CurveEditorTools/Private/Tools/Lattice/Mirror/LatticePointTangentMirrorOp.cpp b/Engine/Plugins/Editor/CurveEditorTools/Source/CurveEditorTools/Private/Tools/Lattice/Mirror/LatticePointTangentMirrorOp.cpp
--- a/Engine/Plugins/Editor/CurveEditorTools/Source/CurveEditorTools/Private/Tools/Lattice/Mirror/LatticePointTangentMirrorOp.cpp
+++ b/Engine/Plugins/Editor/CurveEditorTools/Source/CurveEditorTools/Private/Tools/Lattice/Mirror/LatticePointTangentMirrorOp.cpp
@@ -34,83 +34,93 @@ static void GetDragStartAndMidpointHeights(
 		: ChooseConsistentEdgeVert(EMirroredEdgeType::Top, ControlPoints[TopStartVertIndex], ControlPoints[TopEndVertIndex]);
 }
 }
-	
-TOptional<FLatticePointTangentsMirrorOp> FLatticePointTangentsMirrorOp::MakeFromDraggedControlPoint(
-	int32 InControlPointIndex, const FLatticeDeformer2D& InDeformer, const FCurveEditor& InCurveEditor
-	)
+
+TOptional<FLatticePointTangentsMirrorOp> FLatticePointTangentsMirrorOp::MakeMirrorOpForDragLatticeControlPoint(
+	int32 InControlPointIndex, const FCurveModelID& InCurveId, const TLatticeDeformer2D<FKeyHandle>& InDeformer, const FCurveEditor& InCurveEditor)
 {
-	using namespace MirrorPointDetail;
-	
-	if (InDeformer.NumCells() == 1)
+	TOptional<FCurveMirrorData> CurveData = ComputeTangentMirrorData(InControlPointIndex, InCurveEditor, InCurveId, InDeformer);
+	if (CurveData)
 	{
-		EEdgeSide EdgeSide;
-		double DragStartHeight, MidPointHeight;
-		GetDragStartAndMidpointHeights(InControlPointIndex, InDeformer, EdgeSide, DragStartHeight, MidPointHeight);
+		return FLatticePointTangentsMirrorOp(MoveTemp(*CurveData));
+	}
+	return {};
+}
+
+void FLatticePointTangentsMirrorOp::OnMovePoint(const FVector2D& InControlPoint, const FCurveEditor& CurveEditor)
+{
+	const FCurveModelID& CurveId = CurveData.CurveId;
 		
-		const auto[TopLeftIdx, TopRightIdx, BottomRightIdx, BottomLeftIdx] = InDeformer.GetCellIndices(0);
-		const TConstArrayView<FVector2D> ControlPoints = InDeformer.GetControlPoints();
-		const double MinX = FMath::Min(ControlPoints[TopLeftIdx].X, ControlPoints[BottomLeftIdx].X);
-		const double MaxX = FMath::Min(ControlPoints[TopRightIdx].X, ControlPoints[BottomRightIdx].X);
-		const double Width = MaxX - MinX;
-		TArray<FKeyPosition> Positions; 
-		TMap<FCurveModelID, TArray<float>> CurveToFalloffValues;
+	FCurveModel* CurveModel = CurveEditor.FindCurve(CurveId);
+	if (!CurveModel)
+	{
+		return;
+	}
+	CurveModel->Modify();
 		
-		bool bAtLeastOne = false;
-		CurveEditor::FTangentMirrorSolver TangentSolver(DragStartHeight, MidPointHeight);
-		for (const TPair<FCurveModelID, TArray<FKeyHandle>>& Pair : SplitKeysInCurves(InDeformer.GetKeyIdsInCell(0)))
+	CurveEditor::FCurveTangentMirrorData& MirrorData = CurveData.TangentMirrorSolver.CurveData[CurveId];
+	RecomputeMirroringParallel(CurveEditor, CurveId, MirrorData, InControlPoint.Y,
+		// Flatten the tangents more if their key is closer to the side of the dragged control point, and don't flatten at all on opposite side.
+		// Effectively, this interpolates the interpolated tangents...
+		[this, &MirrorData](int32 KeyIndex, const FVector2D& InterpolatedTangents)
 		{
-			const FCurveModelID CurveId = Pair.Key;
-			// We're going to compute falloff values for each key: between 0 and 1.
-			// The closer the key is to the other side of the lattice grid, the closer the falloff value is to 0.
-			// The closer the key is to the control point, the closer the falloff value is to 1. 
-			const auto ComputeFalloff = [&InCurveEditor, EdgeSide, MinX, Width, &Positions, &CurveToFalloffValues, &CurveId]
-				(const CurveEditor::FMirrorableTangentInfo& InTangentInfo)
-				{
-					const FCurveModel* Curve = InCurveEditor.FindCurve(CurveId);
-					const int32 NumKeys = InTangentInfo.MirrorableKeys.Num();
-					Positions.SetNumUninitialized(NumKeys, EAllowShrinking::No);
-					Curve->GetKeyPositions(InTangentInfo.MirrorableKeys, Positions);
+			const float Alpha = CurveData.FalloffValues[KeyIndex];
+			const float OneMinus = 1.f - Alpha;
+			// Weighted average of interpolated and original tangent value. Effectively, this is a linear falloff.
+			return Alpha * InterpolatedTangents + OneMinus * MirrorData.TangentSolver.InitialValues[KeyIndex];
+		});
+}
 
-					TArray<float>& FalloffValues = CurveToFalloffValues.Add(CurveId);
-					FalloffValues.SetNumUninitialized(NumKeys);
-					for (int32 Index = 0; Index < NumKeys; ++Index)
-					{
-						// If key is on the left = 0, and on right = 1.
-						const float PercentFromLeft = (Positions[Index].InputValue - MinX) / Width;
-						// If control point is on the left		-> key on left side -> falloff = 1,		key on right side -> falloff = 0
-						// If control point is on the right		-> key on left side -> falloff = 0,		key on right side -> falloff = 1
-						const float Falloff = EdgeSide == EEdgeSide::Left ? 1.f - PercentFromLeft : PercentFromLeft;
-						FalloffValues[Index] = Falloff;
-					}
-				};
-			bAtLeastOne |= TangentSolver.AddTangents(InCurveEditor, CurveId, Pair.Value, ComputeFalloff);
-		}
+TOptional<FLatticePointTangentsMirrorOp::FCurveMirrorData> FLatticePointTangentsMirrorOp::ComputeTangentMirrorData(
+	int32 InControlPointIndex, const FCurveEditor& InCurveEditor, const FCurveModelID& InCurveId, const TLatticeDeformer2D<FKeyHandle>& InDeformer
+	)
+{
+	using namespace MirrorPointDetail;
+	if (!ensure(InDeformer.NumCells() == 1))
+	{
+		return {};
+	}
+	
+	EEdgeSide EdgeSide;
+	double DragStartHeight, MidPointHeight;
+	GetDragStartAndMidpointHeights(InControlPointIndex, InDeformer, EdgeSide, DragStartHeight, MidPointHeight);
 		
-		if (bAtLeastOne)
+	const auto[TopLeftIdx, TopRightIdx, BottomRightIdx, BottomLeftIdx] = InDeformer.GetCellIndices(0);
+	const TConstArrayView<FVector2D> ControlPoints = InDeformer.GetControlPoints();
+	const double MinX = FMath::Min(ControlPoints[TopLeftIdx].X, ControlPoints[BottomLeftIdx].X);
+	const double MaxX = FMath::Min(ControlPoints[TopRightIdx].X, ControlPoints[BottomRightIdx].X);
+	const double Width = MaxX - MinX;
+	TArray<float> FalloffValues;
+	
+	// We're going to compute falloff values for each key: between 0 and 1.
+	// The closer the key is to the other side of the lattice grid, the closer the falloff value is to 0.
+	// The closer the key is to the control point, the closer the falloff value is to 1. 
+	const auto ComputeFalloff = [&InCurveEditor, EdgeSide, MinX, Width, &FalloffValues, &InCurveId]
+		(const CurveEditor::FMirrorableTangentInfo& InTangentInfo)
+	{
+		const FCurveModel* Curve = InCurveEditor.FindCurve(InCurveId);
+		const int32 NumKeys = InTangentInfo.MirrorableKeys.Num();
+		TArray<FKeyPosition> Positions; 
+		Positions.SetNumUninitialized(NumKeys, EAllowShrinking::No);
+		Curve->GetKeyPositions(InTangentInfo.MirrorableKeys, Positions);
+
+		FalloffValues.SetNumUninitialized(NumKeys);
+		for (int32 Index = 0; Index < NumKeys; ++Index)
 		{
-			return FLatticePointTangentsMirrorOp(InCurveEditor, MoveTemp(TangentSolver), MoveTemp(CurveToFalloffValues));
+			// If key is on the left = 0, and on right = 1.
+			const float PercentFromLeft = (Positions[Index].InputValue - MinX) / Width;
+			// If control point is on the left		-> key on left side -> falloff = 1,		key on right side -> falloff = 0
+			// If control point is on the right		-> key on left side -> falloff = 0,		key on right side -> falloff = 1
+			const float Falloff = EdgeSide == EEdgeSide::Left ? 1.f - PercentFromLeft : PercentFromLeft;
+			FalloffValues[Index] = Falloff;
 		}
-	}
-	return {};
-}
-
-void FLatticePointTangentsMirrorOp::OnMovePoint(const FVector2D& InControlPoint)
-{
-	for (TPair<FCurveModelID, CurveEditor::FCurveTangentMirrorData>& Pair : TangentMirrorSolver.CurveData)
+	};
+	
+	CurveEditor::FTangentMirrorSolver MirrorSolver(DragStartHeight, MidPointHeight);
+	const bool bHadTangentsToMirror = MirrorSolver.AddTangents(InCurveEditor, InCurveId, InDeformer.GetCellMetaData(0), ComputeFalloff);
+	if (bHadTangentsToMirror) // Only user-specified tangents need mirroring
 	{
-		const FCurveModelID& CurveId = Pair.Key;
-		CurveEditor::FCurveTangentMirrorData& MirrorData = Pair.Value;
-		const TArray<float>& FalloffValues = CurveToFalloffValues[CurveId];
-		RecomputeMirroringParallel(CurveEditor, CurveId, MirrorData, InControlPoint.Y,
-			// Flatten the tangents more if their key is closer to the side of the dragged control point, and don't flatten at all on opposite side.
-			// Effectively, this interpolates the interpolated tangents...
-			[&FalloffValues, &MirrorData](int32 KeyIndex, const FVector2D& InterpolatedTangents)
-			{
-				const float Alpha = FalloffValues[KeyIndex];
-				const float OneMinus = 1.f - Alpha;
-				// Weighted average of interpolated and original tangent value. Effectively, this is a linear falloff.
-				return Alpha * InterpolatedTangents + OneMinus * MirrorData.TangentSolver.InitialValues[KeyIndex];
-			});
+		return FCurveMirrorData(InCurveId, MoveTemp(MirrorSolver), MoveTemp(FalloffValues));
 	}
+	return {};
 }
 }
diff --git a/Engine/Plugins/Editor/CurveEditorTools/Source/CurveEditorTools/Private/Tools/Lattice/Mirror/LatticePointTangentMirrorOp.h b/Engine/Plugins/Editor/CurveEditorTools/Source/CurveEditorTools/Private/Tools/Lattice/Mirror/LatticePointTangentMirrorOp.h
--- a/Engine/Plugins/Editor/CurveEditorTools/Source/CurveEditorTools/Private/Tools/Lattice/Mirror/LatticePointTangentMirrorOp.h
+++ b/Engine/Plugins/Editor/CurveEditorTools/Source/CurveEditorTools/Private/Tools/Lattice/Mirror/LatticePointTangentMirrorOp.h
@@ -4,6 +4,7 @@
 
 #include "HAL/Platform.h"
 #include "Misc/Mirror/TangentMirrorSolver.h"
+#include "Tools/Lattice/LatticeDeformer2D.h"
 
 class FCurveEditor;
 namespace UE::CurveEditor { class FTangentMirrorSolver; }
@@ -13,6 +14,7 @@ template<typename OptionalType> struct TOptional;
 namespace UE::CurveEditorTools
 {
 class FLatticeDeformer2D;
+struct FKeyInCurveId;
 
 /**
  * Provides functionality to drag a control point of a lattice grid and mirror the tangents accordingly.
@@ -26,35 +28,46 @@ class FLatticePointTangentsMirrorOp
 {
 public:
 
-	/** @return Constructed instance that will mirror all keys in InDeformer. Returns valid if in a single celled grid. */
-	static TOptional<FLatticePointTangentsMirrorOp> MakeFromDraggedControlPoint(
-		int32 InControlPointIndex, const FLatticeDeformer2D& InDeformer, const FCurveEditor& InCurveEditor UE_LIFETIMEBOUND
+	/**
+	 * When dragging InControlPointIndex, checks whether the deformer contains keys that need mirroring, which are those that are user specified.
+	 * @return Set if there were keys to mirror when dragging the specifiec control point. Unset otherwise.
+	 */
+	static TOptional<FLatticePointTangentsMirrorOp> MakeMirrorOpForDragLatticeControlPoint(
+		int32 InControlPointIndex, const FCurveModelID& InCurveId, const TLatticeDeformer2D<FKeyHandle>& InDeformer,
+		const FCurveEditor& InCurveEditor UE_LIFETIMEBOUND
 		);
-
+	
 	/** Call with info about where the new control point is located. */
-	void OnMovePoint(const FVector2D& InControlPoint);
+	void OnMovePoint(const FVector2D& InControlPoint, const FCurveEditor& CurveEditor);
 
 private:
+	
+	struct FCurveMirrorData
+	{
+		/** CurveId of the curve whose tangents are being mirrored */
+		FCurveModelID CurveId;
+		/** Mirrors the curve's keys. */
+		CurveEditor::FTangentMirrorSolver TangentMirrorSolver;
+		/** Equal length as TangentMirrorSolver.CurveData.KeyHandles. Contains the falloff values we pre-computed for each key. */
+		TArray<float> FalloffValues;
 
-	/** Used to get and set key attributes */
-	const FCurveEditor& CurveEditor;
-
-	/** Computes the tangents */
-	CurveEditor::FTangentMirrorSolver TangentMirrorSolver;
-
+		explicit FCurveMirrorData(FCurveModelID CurveId, CurveEditor::FTangentMirrorSolver&& TangentMirrorSolver, TArray<float>&& FalloffValues)
+			: CurveId(CurveId), TangentMirrorSolver(MoveTemp(TangentMirrorSolver)), FalloffValues(MoveTemp(FalloffValues))
+		{}
+	};
+	/** Per curve data for performing mirroring. */
+	FCurveMirrorData CurveData;
+	
+	explicit FLatticePointTangentsMirrorOp(FCurveMirrorData&& InCurveData)
+		: CurveData(MoveTemp(InCurveData))
+	{}
+	
 	/**
-	 * Binds the curves to an array containing falloff values for each key. The array is equal lenght as FCurveTangentMirrorData::KeyHandles.
-	 * The falloff values are between 1 and 0: 1 the closer the key is to the dragged control point and 0 on the other side.
-	 * This is used so tangents closer the control point are flattened more.
+	 * Adds all keys that can be mirrored to this operation.
+	 * @return True, if any of the lattice's keys need mirroring.
 	 */
-	TMap<FCurveModelID, TArray<float>> CurveToFalloffValues;
-
-	FLatticePointTangentsMirrorOp(
-		const FCurveEditor& InCurveEditor, CurveEditor::FTangentMirrorSolver InSolver, TMap<FCurveModelID, TArray<float>> CurveToFalloffValues
-		)
-		: CurveEditor(InCurveEditor)
-		, TangentMirrorSolver(MoveTemp(InSolver))
-		, CurveToFalloffValues(MoveTemp(CurveToFalloffValues))
-	{}
+	static TOptional<FCurveMirrorData> ComputeTangentMirrorData(
+		int32 InControlPointIndex, const FCurveEditor& CurveEditor, const FCurveModelID& InCurveId, const TLatticeDeformer2D<FKeyHandle>& InDeformer
+		);
 };
-}
\ No newline at end of file
+}
diff --git a/Engine/Plugins/Editor/CurveEditorTools/Source/CurveEditorTools/Private/Tools/Lattice/Misc/DeformerUtils.cpp b/Engine/Plugins/Editor/CurveEditorTools/Source/CurveEditorTools/Private/Tools/Lattice/Misc/DeformerUtils.cpp
new file mode 100644
--- /dev/null
+++ b/Engine/Plugins/Editor/CurveEditorTools/Source/CurveEditorTools/Private/Tools/Lattice/Misc/DeformerUtils.cpp
@@ -0,0 +1,31 @@
+﻿// Copyright Epic Games, Inc. All Rights Reserved.
+
+#include "DeformerUtils.h"
+
+#include "Misc/Optional.h"
+
+namespace UE::CurveEditorTools
+{
+TOptional<FMoveEdgeToOppositeData> ComputeMoveEdgeToOppositeData(const FLatticeDeformer2D& InDeformer, int32 InEdgeIndex)
+{
+	// At time of implementation, we only have a single cell. Should that change, adjust this implementation.
+	if (InDeformer.NumCells() != 1)
+	{
+		return {};
+	}
+
+	const ELatticeEdgeType EdgeType = GetEdgeTypeFromIndexInSingleCellGrid(InEdgeIndex);
+	if (EdgeType != ELatticeEdgeType::Top && EdgeType != ELatticeEdgeType::Bottom)
+	{
+		return {};
+	}
+
+	const ELatticeEdgeType OppositeEdge = GetOppositeEdge(EdgeType);
+	const int32 OppositeEdgeIndex = GetEdgeIndexInSingleCellGrid(OppositeEdge);
+	const FEdgeVertexIndices Indices = InDeformer.GetEdgeIndices(OppositeEdgeIndex);
+	const FLatticeControlEdge& Edge = InDeformer.GetControlEdge(InEdgeIndex);
+	const TArray<FVector2D, TInlineAllocator<2>> NewEdge = { Edge.Start(), Edge.End() };
+	return FMoveEdgeToOppositeData{ OppositeEdgeIndex, Indices, Edge };
+}
+}
+
diff --git a/Engine/Plugins/Editor/CurveEditorTools/Source/CurveEditorTools/Private/Tools/Lattice/Misc/DeformerUtils.h b/Engine/Plugins/Editor/CurveEditorTools/Source/CurveEditorTools/Private/Tools/Lattice/Misc/DeformerUtils.h
new file mode 100644
--- /dev/null
+++ b/Engine/Plugins/Editor/CurveEditorTools/Source/CurveEditorTools/Private/Tools/Lattice/Misc/DeformerUtils.h
@@ -0,0 +1,22 @@
+﻿// Copyright Epic Games, Inc. All Rights Reserved.
+
+#pragma once
+
+#include "HAL/Platform.h"
+#include "LatticeUtils.h"
+#include "Tools/Lattice/LatticeDeformer2D.h"
+
+template<typename OptionalType> struct TOptional;
+
+namespace UE::CurveEditorTools
+{
+struct FMoveEdgeToOppositeData
+{
+	int32 OppositeEdgeIndex;
+	FEdgeVertexIndices EdgeIndices;
+	FLatticeControlEdge EdgeControlPoints;
+};
+
+/** @return Information about the edge to move to the opposite. Can only move the top or bottom edge. */
+TOptional<FMoveEdgeToOppositeData> ComputeMoveEdgeToOppositeData(const FLatticeDeformer2D& InDeformer, int32 InEdgeIndex);
+}
diff --git a/Engine/Plugins/Editor/CurveEditorTools/Source/CurveEditorTools/Private/Tools/Lattice/Misc/LatticeDrawUtils.cpp b/Engine/Plugins/Editor/CurveEditorTools/Source/CurveEditorTools/Private/Tools/Lattice/Misc/LatticeDrawUtils.cpp
--- a/Engine/Plugins/Editor/CurveEditorTools/Source/CurveEditorTools/Private/Tools/Lattice/Misc/LatticeDrawUtils.cpp
+++ b/Engine/Plugins/Editor/CurveEditorTools/Source/CurveEditorTools/Private/Tools/Lattice/Misc/LatticeDrawUtils.cpp
@@ -25,6 +25,20 @@ FSlateLayoutTransform ComputeViewToViewContainerTransform(const SCurveEditorView
 	return ViewToViewContainerTransform;
 }
 
+FLatticeBounds FLatticeBounds::MakeCombined(const FLatticeBounds& InFirst, const FLatticeBounds& InOther)
+{
+	return FLatticeBounds
+	{
+		InFirst.bIsVisible || InOther.bIsVisible,
+		FVector2d::Min(InFirst.MinValues, InOther.MinValues),
+		FVector2d::Max(InFirst.MaxValues, InOther.MaxValues),
+		FVector2d::Min(InFirst.MinSlatePosition, InOther.MinSlatePosition),
+		FVector2d::Max(InFirst.MaxSlatePosition, InOther.MaxSlatePosition),
+		FVector2d::Min(InFirst.MinValuesCurveSpace, InOther.MinValuesCurveSpace),
+		FVector2d::Max(InFirst.MaxValuesCurveSpace, InOther.MaxValuesCurveSpace)
+	};
+}
+	
 FLatticeBounds ComputeBounds(const FCurveEditor& InCurveEditor)
 {
 	FLatticeBounds Bounds;
@@ -32,54 +46,77 @@ FLatticeBounds ComputeBounds(const FCurveEditor& InCurveEditor)
 	Bounds.MaxValues = { TNumericLimits<double>::Lowest(), TNumericLimits<double>::Lowest() };
 	Bounds.MinSlatePosition = { TNumericLimits<double>::Max(), TNumericLimits<double>::Max() };
 	Bounds.MaxSlatePosition = { TNumericLimits<double>::Lowest(), TNumericLimits<double>::Lowest() };
+	Bounds.MinValuesCurveSpace = { TNumericLimits<double>::Max(), TNumericLimits<double>::Max() };
+	Bounds.MaxValuesCurveSpace = { TNumericLimits<double>::Lowest(), TNumericLimits<double>::Lowest() };
 	
 	const TMap<FCurveModelID, FKeyHandleSet>& SelectedKeySet = InCurveEditor.Selection.GetAll();
 	for (const TPair<FCurveModelID, FKeyHandleSet>& Pair : SelectedKeySet)
 	{
-		const SCurveEditorView* View = InCurveEditor.FindFirstInteractiveView(Pair.Key);
-		if (!View)
-		{
-			continue;
-		}
+		Bounds = FLatticeBounds::MakeCombined(Bounds, ComputeCurveBounds(InCurveEditor, Pair.Key, Pair.Value.AsArray()));
+	}
+	
+	return Bounds;
+}
+	
+FLatticeBounds ComputeCurveBounds(const FCurveEditor& InCurveEditor, const FCurveModelID& InCurveId, TConstArrayView<FKeyHandle> InKeys)
+{
+	FLatticeBounds Bounds;
+	Bounds.MinValues = { TNumericLimits<double>::Max(), TNumericLimits<double>::Max() };
+	Bounds.MaxValues = { TNumericLimits<double>::Lowest(), TNumericLimits<double>::Lowest() };
+	Bounds.MinSlatePosition = { TNumericLimits<double>::Max(), TNumericLimits<double>::Max() };
+	Bounds.MaxSlatePosition = { TNumericLimits<double>::Lowest(), TNumericLimits<double>::Lowest() };
+	Bounds.MinValuesCurveSpace = { TNumericLimits<double>::Max(), TNumericLimits<double>::Max() };
+	Bounds.MaxValuesCurveSpace = { TNumericLimits<double>::Lowest(), TNumericLimits<double>::Lowest() };
 
-		// A newly created view may have a zero-size until the next tick which is a problem if
-		// we ask the View for it's curve space, so we skip over it until it has a size.
-		if(View->GetCachedGeometry().GetLocalSize() == FVector2D::ZeroVector)
-		{
-			continue;
-		}
-		
-		FCurveModel* CurveModel = InCurveEditor.FindCurve(Pair.Key);
-		check(CurveModel);
+	const SCurveEditorView* View = InCurveEditor.FindFirstInteractiveView(InCurveId);
+	if (!View)
+	{
+		return Bounds;
+	}
 
-		TArrayView<const FKeyHandle> KeyHandles = Pair.Value.AsArray();
+	// A newly created view may have a zero-size until the next tick which is a problem if
+	// we ask the View for it's curve space, so we skip over it until it has a size.
+	if(View->GetCachedGeometry().GetLocalSize() == FVector2D::ZeroVector)
+	{
+		return Bounds;
+	}
+	
+	FCurveModel* CurveModel = InCurveEditor.FindCurve(InCurveId);
+	check(CurveModel);
 
-		TArray<FKeyPosition> KeyPositions;
-		KeyPositions.SetNumUninitialized(KeyHandles.Num());
-		CurveModel->GetKeyPositions(KeyHandles, KeyPositions);
+	TArray<FKeyPosition> KeyPositions;
+	KeyPositions.SetNumUninitialized(InKeys.Num());
+	CurveModel->GetKeyPositions(InKeys, KeyPositions);
+	
+	const FTransform2d AbsToCurveSpace = View->GetViewToCurveTransform(InCurveId);
+	double M00, M01, M10, M11;
+	AbsToCurveSpace.GetMatrix().GetMatrix(M00, M01, M10, M11);
 
-		FCurveEditorScreenSpace ViewSpace = View->GetViewSpace();
-		const FSlateLayoutTransform ViewToViewContainerTransform = ComputeViewToViewContainerTransform(*View, InCurveEditor);
-		for (int32 i = 0; i < KeyPositions.Num(); ++i)
-		{
-			const FKeyPosition& KeyPosition = KeyPositions[i];
-			const FVector2D Position_ViewSpace = FVector2D(
-				ViewSpace.SecondsToScreen(KeyPosition.InputValue), ViewSpace.ValueToScreen(KeyPosition.OutputValue)
-				);
-			const FVector2D PanelSpaceLocation = ViewToViewContainerTransform.TransformPoint(Position_ViewSpace);
+	FCurveEditorScreenSpace ViewSpace = View->GetCurveSpace(InCurveId);
+	const FSlateLayoutTransform ViewToViewContainerTransform = ComputeViewToViewContainerTransform(*View, InCurveEditor);
+	for (int32 i = 0; i < KeyPositions.Num(); ++i)
+	{
+		const FKeyPosition& KeyPosition = KeyPositions[i];
+		const FVector2D Position_ViewSpace = FVector2D(
+			ViewSpace.SecondsToScreen(KeyPosition.InputValue), ViewSpace.ValueToScreen(KeyPosition.OutputValue)
+			);
+		const FVector2D PanelSpaceLocation = ViewToViewContainerTransform.TransformPoint(Position_ViewSpace);
 
-			Bounds.bIsVisible = true;
-			
-			Bounds.MinValues.X = FMath::Min(KeyPosition.InputValue, Bounds.MinValues.X);
-			Bounds.MinValues.Y = FMath::Min(KeyPosition.OutputValue, Bounds.MinValues.Y);
-			Bounds.MaxValues.X = FMath::Max(KeyPosition.InputValue, Bounds.MaxValues.X);
-			Bounds.MaxValues.Y = FMath::Max(KeyPosition.OutputValue, Bounds.MaxValues.Y);
-			
-			Bounds.MinSlatePosition = FVector2D::Min(PanelSpaceLocation, Bounds.MinSlatePosition);
-			Bounds.MaxSlatePosition = FVector2D::Max(PanelSpaceLocation, Bounds.MaxSlatePosition);
-		}
+		Bounds.bIsVisible = true;
+		
+		Bounds.MinValues.X = FMath::Min(KeyPosition.InputValue, Bounds.MinValues.X);
+		Bounds.MinValues.Y = FMath::Min(KeyPosition.OutputValue, Bounds.MinValues.Y);
+		Bounds.MaxValues.X = FMath::Max(KeyPosition.InputValue, Bounds.MaxValues.X);
+		Bounds.MaxValues.Y = FMath::Max(KeyPosition.OutputValue, Bounds.MaxValues.Y);
+		
+		Bounds.MinSlatePosition = FVector2D::Min(PanelSpaceLocation, Bounds.MinSlatePosition);
+		Bounds.MaxSlatePosition = FVector2D::Max(PanelSpaceLocation, Bounds.MaxSlatePosition);
+
+		const FVector2D CurveSpacePosition = TransformAbsoluteToCurveSpace(AbsToCurveSpace, FVector2D{ KeyPosition.InputValue, KeyPosition.OutputValue });
+		Bounds.MinValuesCurveSpace = FVector2D::Min(CurveSpacePosition, Bounds.MinValuesCurveSpace);
+		Bounds.MaxValuesCurveSpace = FVector2D::Max(CurveSpacePosition, Bounds.MaxValuesCurveSpace);
 	}
-	
+
 	return Bounds;
 }
 
@@ -367,4 +404,4 @@ void DrawLatticeControls(
 	// ... and control points draw over edges.
 	DrawControlPoints(InDrawData, InViewContainerGeometry, OutDrawElements, InPaintOnLayerId);
 }
-}
+}
\ No newline at end of file
diff --git a/Engine/Plugins/Editor/CurveEditorTools/Source/CurveEditorTools/Private/Tools/Lattice/Misc/LatticeDrawUtils.h b/Engine/Plugins/Editor/CurveEditorTools/Source/CurveEditorTools/Private/Tools/Lattice/Misc/LatticeDrawUtils.h
--- a/Engine/Plugins/Editor/CurveEditorTools/Source/CurveEditorTools/Private/Tools/Lattice/Misc/LatticeDrawUtils.h
+++ b/Engine/Plugins/Editor/CurveEditorTools/Source/CurveEditorTools/Private/Tools/Lattice/Misc/LatticeDrawUtils.h
@@ -6,11 +6,13 @@
 #include "HAL/Platform.h"
 #include "Math/Vector2D.h"
 
+struct FCurveModelID;
 class FCurveEditor;
 class FSlateLayoutTransform;
 class FSlateWindowElementList;
 class SCurveEditorView;
 struct FGeometry;
+struct FKeyHandle;
 
 /** Put utils specific to drawing the lattice controls into this file */
 namespace UE::CurveEditorTools
@@ -32,17 +34,27 @@ struct FLatticeBounds
 	/** If false, then do not discard the other values. */
 	bool bIsVisible = false;
 
-	// Values in graph
+	// Absolute key values, i.e. FKeyPosition::InputValue, FKeyPosition::OutputValue.
 	FVector2D MinValues;
 	FVector2D MaxValues;
 
-	// Values on screen
+	// Values on screen, but transformed to screen space.
 	FVector2D MinSlatePosition;
 	FVector2D MaxSlatePosition;
+
+	// Axis values in curve space of the current view mode.
+	// For example, in normalized mode, where the Y-axis only has values from 0 to 1, a Y could be 0.8, etc.
+	FVector2d MinValuesCurveSpace;
+	FVector2d MaxValuesCurveSpace;
+
+	static FLatticeBounds MakeCombined(const FLatticeBounds& InFirst, const FLatticeBounds& InOther);
 };
+
 /** Computes the min and max slate and value points in the user's selection. */
 FLatticeBounds ComputeBounds(const FCurveEditor& InCurveEditor);
-
+/** Computes the bounds just for the given curve. */
+FLatticeBounds ComputeCurveBounds(const FCurveEditor& InCurveEditor, const FCurveModelID& InCurveId, TConstArrayView<FKeyHandle> InKeys);
+	
 /**
  * Transforms the values of the control points to where they should be drawn in InLatticeGeometry.
  * @param InCurveEditor Needed to get interactive view for transforming values.
diff --git a/Engine/Plugins/Editor/CurveEditorTools/Source/CurveEditorTools/Private/Tools/Lattice/Misc/LatticeUtils.cpp b/Engine/Plugins/Editor/CurveEditorTools/Source/CurveEditorTools/Private/Tools/Lattice/Misc/LatticeUtils.cpp
--- a/Engine/Plugins/Editor/CurveEditorTools/Source/CurveEditorTools/Private/Tools/Lattice/Misc/LatticeUtils.cpp
+++ b/Engine/Plugins/Editor/CurveEditorTools/Source/CurveEditorTools/Private/Tools/Lattice/Misc/LatticeUtils.cpp
@@ -86,32 +86,55 @@ FCellVertexIndices GetCellIndices(int32 InCellIndex, int32 InNumPointsInWidth)
 	return CellIndices;
 }
 
-void AddSelectionToDeformer(FLatticeDeformer2D& InDeformer, const FCurveEditor& InCurveEditor)
+bool IsLatticeTooSmall(const FVector2D& InMin, const FVector2D& InMax)
 {
-	for (const TPair<FCurveModelID, FKeyHandleSet>& Pair : InCurveEditor.Selection.GetAll())
+	const FVector2D DeltaValues = InMax - InMin;
+	const bool bTooSmall = DeltaValues.X < UE_KINDA_SMALL_NUMBER || DeltaValues.Y < UE_KINDA_SMALL_NUMBER;
+	return bTooSmall;
+}
+	
+TArray<FVector2D> GenerateControlPoints(int32 InNumCellsInX, int32 InNumCellsInY, const FVector2D& InBottomLeft, const FVector2D& InTopRight)
+{
+	TArray<FVector2D> ControlPoints;
+	// Lattice with 0 width or height does not make sense. 
+	if (!ensure(InBottomLeft.X < InTopRight.X) || !ensure(InBottomLeft.Y < InTopRight.Y))
 	{
-		const FCurveModelID CurveId = Pair.Key;
-		const FCurveModel* CurveModel = InCurveEditor.FindCurve(CurveId);
-		if (!CurveModel)
-		{
-			continue;
-		}
+		return ControlPoints;
+	}
 
-		// Looping like this avoids allocating a big FKeyPosition array.
-		TConstArrayView<FKeyHandle> Keys = Pair.Value.AsArray();
-		for (int32 Index = 0; Index < Keys.Num(); ++Index)
+	const int32 NumPointsInWidth = InNumCellsInX + 1;
+	const int32 NumPointsInHeight = InNumCellsInY + 1;
+	const int32 NumControlPoints = NumPointsInWidth * NumPointsInHeight;
+	if (NumControlPoints <= 0)
+	{
+		return ControlPoints;
+	}
+	ControlPoints.SetNumUninitialized(NumControlPoints);
+
+	const double DeltaX = (InTopRight.X - InBottomLeft.X) / static_cast<double>(InNumCellsInX);
+	const double DeltaY = (InTopRight.Y - InBottomLeft.Y) / static_cast<double>(InNumCellsInY);
+	for (int32 X = 0; X < NumPointsInWidth; ++X)
+	{
+		for (int32 Y = 0; Y < NumPointsInWidth; ++Y)
 		{
-			const FKeyHandle& KeyId = Keys[Index];
-				
-			FKeyPosition Position;
-			TConstArrayView<FKeyHandle> SingleKeyView(&KeyId, 1);
-			TArrayView<FKeyPosition> SinglePositionView(&Position, 1);
-			CurveModel->GetKeyPositions(SingleKeyView, SinglePositionView);
-
-			const FVector2D AsVector{ Position.InputValue, Position.OutputValue };
-			const TConstArrayView<FVector2D> KeyView(&AsVector, 1);
-			InDeformer.AddPoints_BeforeLatticeMoved({ FKeyInCurveId{ KeyId, CurveId } }, KeyView);
+			const int32 Index = MatrixIndicesToFlatIndex(X, Y, NumPointsInWidth);
+			const FVector2D ControlPosition(InBottomLeft.X + DeltaX * X, InBottomLeft.Y + DeltaY * Y);
+			ControlPoints[Index] = ControlPosition;
 		}
 	}
+
+	return ControlPoints;
+}
+
+TArray<FVector2D> TransformPoints(const FTransform2d& InTransform, TConstArrayView<FVector2D> InControlPoints)
+{
+	TArray<FVector2D> TransformedPoints;
+	TransformedPoints.Reserve(InControlPoints.Num());
+	Algo::Transform(InControlPoints, TransformedPoints, [&InTransform](const FVector2D& ControlPoint)
+	{
+		const FVector2D Transformed = InTransform.TransformPoint(ControlPoint);
+		return Transformed;
+	});
+	return TransformedPoints;
 }
 }
diff --git a/Engine/Plugins/Editor/CurveEditorTools/Source/CurveEditorTools/Private/Tools/Lattice/Misc/LatticeUtils.h b/Engine/Plugins/Editor/CurveEditorTools/Source/CurveEditorTools/Private/Tools/Lattice/Misc/LatticeUtils.h
--- a/Engine/Plugins/Editor/CurveEditorTools/Source/CurveEditorTools/Private/Tools/Lattice/Misc/LatticeUtils.h
+++ b/Engine/Plugins/Editor/CurveEditorTools/Source/CurveEditorTools/Private/Tools/Lattice/Misc/LatticeUtils.h
@@ -2,6 +2,7 @@
 
 #pragma once
 
+#include "Math/TransformCalculus2D.h"
 #include "Misc/AssertionMacros.h"
 #include "Misc/CoreMiscDefines.h"
 #include "HAL/Platform.h"
@@ -12,7 +13,8 @@ template<typename OptionalType> struct TOptional;
 /** Put utils specific to lattice deforming into this file */
 namespace UE::CurveEditorTools
 {
-class FLatticeDeformer2D;
+template<typename T> class TLatticeDeformer2D;
+struct FKeyInCurveId;
 	
 /** @return Index you can use in a TArray representing a 2D matrix. */
 int32 MatrixIndicesToFlatIndex(int32 InColumn, int32 InRow, int32 InNumColumns);
@@ -68,9 +70,6 @@ struct FCellVertexIndices
  */
 FCellVertexIndices GetCellIndices(int32 InCellIndex, int32 InNumPointsInWidth);
 
-/** Adds InCurveEditor's selection to InDeformer. */
-void AddSelectionToDeformer(FLatticeDeformer2D& InDeformer, const FCurveEditor& InCurveEditor);
-
 enum class ELatticeEdgeType : uint8
 {
 	Top,
@@ -85,6 +84,15 @@ constexpr ELatticeEdgeType GetEdgeTypeFromIndexInSingleCellGrid(int32 InEdgeInde
 	
 /** @return The opposite edge of the passed in edge */
 constexpr ELatticeEdgeType GetOppositeEdge(ELatticeEdgeType Type);
+
+/** @return Whether a lattice can be formed by these points. */
+bool IsLatticeTooSmall(const FVector2D& InMin, const FVector2D& InMax);
+	
+/** @return The lattice control points given the construction arguments. Use GetCellIndices to access cell control points. */
+TArray<FVector2D> GenerateControlPoints(int32 InNumCellsInX, int32 InNumCellsInY, const FVector2D& InBottomLeft, const FVector2D& InTopRight);
+
+/** @return InControlPoints transformed by InTransform. */
+TArray<FVector2D> TransformPoints(const FTransform2d& InTransform, TConstArrayView<FVector2D> InControlPoints);
 }
 
 namespace UE::CurveEditorTools
diff --git a/Engine/Plugins/Editor/CurveEditorTools/Source/CurveEditorTools/Private/Tools/Lattice/PerCurveLatticeData.cpp b/Engine/Plugins/Editor/CurveEditorTools/Source/CurveEditorTools/Private/Tools/Lattice/PerCurveLatticeData.cpp
new file mode 100644
--- /dev/null
+++ b/Engine/Plugins/Editor/CurveEditorTools/Source/CurveEditorTools/Private/Tools/Lattice/PerCurveLatticeData.cpp
@@ -0,0 +1,66 @@
+﻿// Copyright Epic Games, Inc. All Rights Reserved.
+
+#include "PerCurveLatticeData.h"
+
+#include "Containers/Map.h"
+#include "CurveEditor.h"
+#include "SCurveEditorView.h"
+#include "Misc/LatticeDrawUtils.h"
+#include "Misc/LatticeUtils.h"
+#include "Misc/VectorMathUtils.h"
+
+namespace UE::CurveEditorTools
+{
+namespace PerCurveDetail
+{
+static void AddPointsToDeformer(TConstArrayView<FKeyHandle> InKeys, const FCurveModel* CurveModel, FPerCurveDeformer2D& LocalDeformer)
+{
+	// Looping like this avoids allocating a big FKeyPosition array.
+	for (int32 Index = 0; Index < InKeys.Num(); ++Index)
+	{
+		const FKeyHandle& KeyId = InKeys[Index];
+		
+		FKeyPosition Position;
+		TConstArrayView<FKeyHandle> SingleKeyView(&KeyId, 1);
+		TArrayView<FKeyPosition> SinglePositionView(&Position, 1);
+		CurveModel->GetKeyPositions(SingleKeyView, SinglePositionView);
+
+		const FVector2D AsVector{ Position.InputValue, Position.OutputValue };
+		const TConstArrayView<FVector2D> KeyView(&AsVector, 1);
+		LocalDeformer.AddPoints_BeforeLatticeMoved({ KeyId }, KeyView);
+	}
+}
+}
+	
+TMap<FCurveModelID, FPerCurveLatticeData> BuildPerLatticeData(const FLatticeBounds& Lattice, const FCurveEditor& InCurveEditor)
+{
+	TMap<FCurveModelID, FPerCurveLatticeData> PerCurveData;
+	for (const TPair<FCurveModelID, FKeyHandleSet>& SelectionPair : InCurveEditor.Selection.GetAll())
+	{
+		const FCurveModelID& CurveId = SelectionPair.Key;
+		const SCurveEditorView* View = InCurveEditor.FindFirstInteractiveView(CurveId);
+		const FCurveModel* CurveModel = InCurveEditor.FindCurve(CurveId);
+		if (!View || !CurveModel)
+		{
+			continue;
+		}
+		
+		// 1. Need to convert Lattice.MinValuesCurveSpace & MaxValuesCurveSpace to absolute for this particular curve.
+		const FTransform2d ViewToCurveTransform = View->GetViewToCurveTransform(CurveId);
+		const FVector2D LocalMin = TransformCurveSpaceToAbsolute(ViewToCurveTransform, Lattice.MinValuesCurveSpace);
+		const FVector2D LocalMax = TransformCurveSpaceToAbsolute(ViewToCurveTransform, Lattice.MaxValuesCurveSpace);
+		// 2. Then, compute the transform to convert from NormalizedDeformer to LocalDeformer.
+		const FTransform2d ToCurveSpace = TransformRectBetweenSpaces(Lattice.MinValues, Lattice.MaxValues, LocalMin, LocalMax);
+
+		// In certain cases, the lattice may be too tiny after transformation (lattice maths would divide by zero - the constructor would ensure()).
+		if (!IsLatticeTooSmall(LocalMin, LocalMax))
+		{
+			FPerCurveLatticeData& CurveData = PerCurveData.Add(
+				CurveId, FPerCurveLatticeData{  FPerCurveDeformer2D(1, 1, LocalMin, LocalMax), ToCurveSpace }
+			);
+			PerCurveDetail::AddPointsToDeformer(SelectionPair.Value.AsArray(), CurveModel, CurveData.CurveDeformer);
+		}
+	}
+	return PerCurveData;
+}
+}
diff --git a/Engine/Plugins/Editor/CurveEditorTools/Source/CurveEditorTools/Private/Tools/Lattice/PerCurveLatticeData.h b/Engine/Plugins/Editor/CurveEditorTools/Source/CurveEditorTools/Private/Tools/Lattice/PerCurveLatticeData.h
new file mode 100644
--- /dev/null
+++ b/Engine/Plugins/Editor/CurveEditorTools/Source/CurveEditorTools/Private/Tools/Lattice/PerCurveLatticeData.h
@@ -0,0 +1,52 @@
+﻿// Copyright Epic Games, Inc. All Rights Reserved.
+
+#pragma once
+
+#include "Containers/ContainersFwd.h"
+#include "LatticeDeformer2D.h"
+#include "LatticeFwd.h"
+#include "Math/TransformCalculus2D.h"
+
+class FCurveEditor;
+struct FCurveModelID;
+struct FKeyHandle;
+
+namespace UE::CurveEditorTools
+{
+	struct FLatticeBounds;
+
+	/** Data stored per curve in FCurveEditorLatticeTool. */
+struct FPerCurveLatticeData
+{
+	/**
+	 * Deformers for each curve; the control points are global deformer's control points transformed by the per-curve transform.
+	 * The per-curve deformers' control points are always in absolute key space (i.e. exactly the values for FKeyPosition::InputValue and OutputValue).
+	 */
+	FPerCurveDeformer2D CurveDeformer;
+	
+	/**
+	 * Transform from the global deformer's control points to per-curve lattice control points.
+	 *
+	 * The transform depends on the view:
+	 * - In Absolute view, this is effectively an identity matrix.
+	 * - For Normalized view, let's use an example:
+	 *		- Suppose the following curves:
+	 *			Curve A: global min and max absolute values are y=0 and y=100, respectively.
+	 *			Curve B: global min and max absolute values are y=50 and y=150, respectively.
+	 *		- Suppose the user selects keys such that Deformer is placed with bottom & top edge at y=0.6 & y=0.8 respectively.
+	 *		- Deformer's bottom & top edge's absolute values would be y=60 and y=130 (--> the min & max absolute values across all selected curves).
+	 *		- The per-curve lattices would be placed as follows
+	 *			Curve A: lattice bottom & top edge at y=60 and y=80, respectively.
+	 *			Curve B: lattice bottom & top edge at y=110 and y=130, respectively.
+	 *		- The per-curve transforms are computed such that the Deformer's lattice with bottom & top edge at y=60 and y=130 and transformed to
+	 *		the local lattices of Curve A (y=60 and y=80) and Curve B(y=110 and y=130), respectively.
+	 *		
+	 * The computation of the per-curve transform uses the per-curve transform of SCurveEditorView::GetViewToCurveTransform, which implements
+	 * Absolute, Normalized, etc. views.
+	 */
+	FTransform2d GlobalDeformerToCurveDeformer;
+};
+
+/** Builds the per curve data for lattice. */
+TMap<FCurveModelID, FPerCurveLatticeData> BuildPerLatticeData(const FLatticeBounds& Lattice, const FCurveEditor& InCurveEditor);
+}
diff --git a/Engine/Plugins/Editor/CurveEditorTools/Source/CurveEditorTools/Private/Tools/Lattice/UObject/LatticeUndoData.cpp b/Engine/Plugins/Editor/CurveEditorTools/Source/CurveEditorTools/Private/Tools/Lattice/UObject/LatticeUndoData.cpp
deleted file mode 100644
diff --git a/Engine/Plugins/Editor/CurveEditorTools/Source/CurveEditorTools/Private/Tools/Lattice/UObject/LatticeUndoData.h b/Engine/Plugins/Editor/CurveEditorTools/Source/CurveEditorTools/Private/Tools/Lattice/UObject/LatticeUndoData.h
deleted file mode 100644
diff --git a/Engine/Plugins/Editor/CurveEditorTools/Source/CurveEditorTools/Private/Tools/Lattice/Undo/LatticeSnapshot.cpp b/Engine/Plugins/Editor/CurveEditorTools/Source/CurveEditorTools/Private/Tools/Lattice/Undo/LatticeSnapshot.cpp
new file mode 100644
--- /dev/null
+++ b/Engine/Plugins/Editor/CurveEditorTools/Source/CurveEditorTools/Private/Tools/Lattice/Undo/LatticeSnapshot.cpp
@@ -0,0 +1,140 @@
+﻿// Copyright Epic Games, Inc. All Rights Reserved.
+
+#include "LatticeSnapshot.h"
+
+#include "Algo/Transform.h"
+#include "CurveEditor.h"
+#include "SCurveEditorPanel.h"
+#include "Tools/Lattice/Misc/LatticeUtils.h"
+
+namespace UE::CurveEditorTools
+{
+FArchive& operator<<(FArchive& InArchive, FTransform2d& InTransform)
+{
+	const TMatrix2x2<double>& Matrix = InTransform.GetMatrix();
+	FVector2D Translation = InTransform.GetTranslation();
+	double M00, M01, M10, M11;
+	Matrix.GetMatrix(M00, M01, M10, M11);
+	
+	InArchive << M00 << M01 << M10 << M11;
+	InArchive << Translation;
+
+	if (InArchive.IsLoading())
+	{
+		InTransform = FTransform2d(TMatrix2x2<double>(M00, M01, M10, M11), Translation);
+	}
+	return InArchive;
+}
+	
+FArchive& operator<<(FArchive& InArchive, FPerCurveSnapshot& InSnapshot)
+{
+	InArchive << InSnapshot.GlobalToLocalTransform;
+	InArchive << InSnapshot.KeyCoords;
+	InArchive << InSnapshot.Keys;
+	return InArchive;
+}
+	
+bool FLatticeSnapshot::CanApplySnapshot(const FCurveEditor& InCurveEditor) const
+{
+	return InCurveEditor.Selection.GetSerialNumber() == SelectionSerialNumber
+		&& ViewMode && InCurveEditor.GetPanel() && InCurveEditor.GetPanel()->GetViewMode() == *ViewMode;
+}
+	
+FArchive& operator<<(FArchive& InArchive, FLatticeSnapshot& InSnapshot)
+{
+	InArchive << InSnapshot.ViewMode;
+	InArchive << InSnapshot.SelectionSerialNumber;
+	InArchive << InSnapshot.ControlPointToCurveSpace;
+	InArchive << InSnapshot.GlobalControlPoints;
+	InArchive << InSnapshot.PerCurveData;
+	return InArchive;
+}
+
+bool ApplySnapshot(
+	const FLatticeSnapshot& InSnapshot, const FCurveEditor& InCurveEditor,
+	FTransform2d& OutControlPointToCurveSpace, FGlobalLatticeDeformer2D& OutGlobalDeformer,
+	TMap<FCurveModelID, FPerCurveLatticeData>& OutPerCurveData, uint32& OutSerialNumber
+)
+{
+	if (!InSnapshot.CanApplySnapshot(InCurveEditor)
+		// To make implementation easier, we just handle 1 celled lattices for now. If you change that in the future, adjust the below code.
+		|| !ensure(OutGlobalDeformer.NumCells() == 1))
+	{
+		return false;
+	}
+
+	// Properties for global deformers
+	OutControlPointToCurveSpace = InSnapshot.ControlPointToCurveSpace;
+	for (int32 PointIndex = 0; PointIndex < InSnapshot.GlobalControlPoints.Num(); ++PointIndex)
+    {
+    	TConstArrayView<FVector2D> ControlPoint(&InSnapshot.GlobalControlPoints[PointIndex], 1);
+    	OutGlobalDeformer.SetControlsPointNoRecompute({ PointIndex }, { ControlPoint });
+    }
+
+	// Per-curve data
+	OutPerCurveData.Empty(InSnapshot.PerCurveData.Num());
+    for (const TPair<FCurveModelID, FPerCurveSnapshot>& Pair : InSnapshot.PerCurveData)
+    {
+    	const FPerCurveSnapshot& Snapshot = Pair.Value;
+    	const TArray<FVector2D> ControlPoints = TransformPoints(Snapshot.GlobalToLocalTransform, OutGlobalDeformer.GetControlPoints());
+    	constexpr int32 CellIndexZero = 0;
+    	const auto[TopLeft, TopRight, BottomRight, BottomLeft] = GetCellIndices(CellIndexZero, 2/*NumPoints*/);
+    	
+    	const FVector2D Min = FVector2D::Min(FVector2D::Min(ControlPoints[0], ControlPoints[1]), FVector2D::Min(ControlPoints[2], ControlPoints[3]));
+    	const FVector2D Max = FVector2D::Max(FVector2D::Max(ControlPoints[0], ControlPoints[1]), FVector2D::Max(ControlPoints[2], ControlPoints[3]));
+    	// This *might* happen but be really irrealistic
+    	if (IsLatticeTooSmall(Min, Max))
+    	{
+    		continue;
+    	}
+    	
+    	FPerCurveLatticeData Data
+    	{
+    		.CurveDeformer = FPerCurveDeformer2D(1, 1, Min, Max),
+    		.GlobalDeformerToCurveDeformer = Snapshot.GlobalToLocalTransform
+    	};
+
+    	FPerCurveDeformer2D& Deformer = Data.CurveDeformer;
+    	Deformer.SetControlsPointNoRecompute(
+    		{ TopLeft, TopRight, BottomRight, BottomLeft },
+    		{ ControlPoints[TopLeft], ControlPoints[TopRight], ControlPoints[BottomRight], ControlPoints[BottomLeft] }
+    		);
+    	Deformer.SetKeyCoordsInCell(CellIndexZero, Snapshot.KeyCoords);
+    	Deformer.SetCellMetaData(CellIndexZero, Snapshot.Keys);
+    	
+    	OutPerCurveData.Add(Pair.Key, MoveTemp(Data));
+    }
+
+	OutSerialNumber = InSnapshot.SelectionSerialNumber;
+	
+	return true;
+}
+
+FLatticeSnapshot TakeSnapshot(
+	const FCurveEditor& InCurveEditor,
+	const FTransform2d& InControlPointToCurveSpace, const FGlobalLatticeDeformer2D& InGlobalDeformer,
+	const TMap<FCurveModelID, FPerCurveLatticeData>& InPerCurveData
+	)
+{
+	FLatticeSnapshot Snapshot
+	{
+		.ViewMode = InCurveEditor.GetPanel() ? InCurveEditor.GetPanel()->GetViewMode() : TOptional<ECurveEditorViewID>(),
+		.SelectionSerialNumber = InCurveEditor.Selection.GetSerialNumber(),
+		.ControlPointToCurveSpace = InControlPointToCurveSpace,
+		.GlobalControlPoints = TArray<FVector2D>(InGlobalDeformer.GetControlPoints())
+	};
+
+	Algo::Transform(InPerCurveData, Snapshot.PerCurveData, [](const TPair<FCurveModelID, FPerCurveLatticeData>& Pair)
+	{
+		FPerCurveSnapshot CurveSnapshot
+		{
+			.GlobalToLocalTransform = Pair.Value.GlobalDeformerToCurveDeformer,
+			.KeyCoords = TArray<FLatticeKeyCoords>(Pair.Value.CurveDeformer.GetKeyCoordsInCell(0)),
+			.Keys = TArray<FKeyHandle>(Pair.Value.CurveDeformer.GetCellMetaData(0))
+		};
+		return TPair<FCurveModelID, FPerCurveSnapshot>(Pair.Key, MoveTemp(CurveSnapshot));
+	});
+	
+	return Snapshot;
+}
+}
\ No newline at end of file
diff --git a/Engine/Plugins/Editor/CurveEditorTools/Source/CurveEditorTools/Private/Tools/Lattice/Undo/LatticeSnapshot.h b/Engine/Plugins/Editor/CurveEditorTools/Source/CurveEditorTools/Private/Tools/Lattice/Undo/LatticeSnapshot.h
new file mode 100644
--- /dev/null
+++ b/Engine/Plugins/Editor/CurveEditorTools/Source/CurveEditorTools/Private/Tools/Lattice/Undo/LatticeSnapshot.h
@@ -0,0 +1,61 @@
+﻿// Copyright Epic Games, Inc. All Rights Reserved.
+
+#pragma once
+
+#include "CurveEditorTypes.h"
+#include "Misc/Optional.h"
+#include "Tools/Lattice/LatticeFwd.h"
+#include "Tools/Lattice/PerCurveLatticeData.h"
+
+class FCurveEditor;
+enum class ECurveEditorViewID : uint64;
+
+namespace UE::CurveEditorTools
+{
+struct FPerCurveSnapshot
+{
+	FTransform2d GlobalToLocalTransform;
+	TArray<FLatticeKeyCoords> KeyCoords;
+	TArray<FKeyHandle> Keys;
+	
+	friend FArchive& operator<<(FArchive& InArchive, FPerCurveSnapshot& InSnapshot);
+};
+	
+struct FLatticeSnapshot
+{
+	/** The view mode the editor was in when during this snapshot. The snapshot can only be applied if the editor is currently in the same view mode. */
+	TOptional<ECurveEditorViewID> ViewMode;
+	/** The selection serial number the editor had when this snapshot was taken. The snapshot can only be applied the selection states are the same. */
+	uint32 SelectionSerialNumber;
+
+	/** Corresponds to FLatticeDeformerState::ControlPointToCurveSpace. */
+	FTransform2d ControlPointToCurveSpace;
+	/** The controls points the global lattice had. Corresponds to FLatticeDeformerState::GlobalDeformer. */
+	TArray<FVector2D> GlobalControlPoints;
+	/** Corresponds to FLatticeDeformerState::PerCurveData. */
+	TMap<FCurveModelID, FPerCurveSnapshot> PerCurveData;
+	
+	/** @return Whether the saved data is compatible with the currentl state of InCurveEditor. */
+	bool CanApplySnapshot(const FCurveEditor& InCurveEditor) const;
+
+	friend FArchive& operator<<(FArchive& InArchive, FLatticeSnapshot& InSnapshot);
+};
+
+/**
+ * Tries to apply the snapshot if possible.
+ * @return Whether the data was applied.
+ */
+bool ApplySnapshot(
+	const FLatticeSnapshot& InSnapshot, const FCurveEditor& InCurveEditor,
+	FTransform2d& OutControlPointToCurveSpace, FGlobalLatticeDeformer2D& OutGlobalDeformer,
+	TMap<FCurveModelID, FPerCurveLatticeData>& OutPerCurveData, uint32& OutSerialNumber
+	);
+
+/** Takes a snapshot of the lattice tool and the curve editor's data required in the future to validate that the data can be applied. */
+FLatticeSnapshot TakeSnapshot(
+	const FCurveEditor& InCurveEditor,
+	const FTransform2d& InControlPointToCurveSpace, const FGlobalLatticeDeformer2D& InGlobalDeformer,
+	const TMap<FCurveModelID, FPerCurveLatticeData>& InPerCurveData
+	);
+}
+
diff --git a/Engine/Plugins/Editor/CurveEditorTools/Source/CurveEditorTools/Private/Tools/Lattice/Undo/LatticeUndoObject.cpp b/Engine/Plugins/Editor/CurveEditorTools/Source/CurveEditorTools/Private/Tools/Lattice/Undo/LatticeUndoObject.cpp
new file mode 100644
--- /dev/null
+++ b/Engine/Plugins/Editor/CurveEditorTools/Source/CurveEditorTools/Private/Tools/Lattice/Undo/LatticeUndoObject.cpp
@@ -0,0 +1,17 @@
+﻿// Copyright Epic Games, Inc. All Rights Reserved.
+
+#include "LatticeUndoObject.h"
+
+#include "CurveEditor.h"
+
+void UCurveEditorTools_LatticeUndoObject::Serialize(FArchive& Ar)
+{
+	Super::Serialize(Ar);
+	Ar << Snapshot;
+}
+
+void UCurveEditorTools_LatticeUndoObject::PostEditUndo()
+{
+	UObject::PostEditUndo();
+	OnPostEditUndo.Broadcast();
+}
diff --git a/Engine/Plugins/Editor/CurveEditorTools/Source/CurveEditorTools/Private/Tools/Lattice/Undo/LatticeUndoObject.h b/Engine/Plugins/Editor/CurveEditorTools/Source/CurveEditorTools/Private/Tools/Lattice/Undo/LatticeUndoObject.h
new file mode 100644
--- /dev/null
+++ b/Engine/Plugins/Editor/CurveEditorTools/Source/CurveEditorTools/Private/Tools/Lattice/Undo/LatticeUndoObject.h
@@ -0,0 +1,28 @@
+﻿// Copyright Epic Games, Inc. All Rights Reserved.
+
+#pragma once
+
+#include "LatticeSnapshot.h"
+#include "Delegates/Delegate.h"
+#include "UObject/Object.h"
+#include "LatticeUndoObject.generated.h"
+
+/** Stores data for the purposes of undo / redo in the lattice tool. */
+UCLASS()
+class UCurveEditorTools_LatticeUndoObject : public UObject
+{
+	GENERATED_BODY()
+public:
+
+	DECLARE_MULTICAST_DELEGATE(FOnPostEditUndo);
+	/** Invoked when the user has performed an undo / redo operation. You should call RestoreState. */
+	FOnPostEditUndo OnPostEditUndo;
+
+	/** Data that is undone / redone by the transaction sytem. */
+	UE::CurveEditorTools::FLatticeSnapshot Snapshot;
+
+	//~ Begin UObject Interface
+	virtual void Serialize(FArchive& Ar) override;
+	virtual void PostEditUndo() override;
+	//~ End UObject Interface
+};
diff --git a/Engine/Source/Editor/CurveEditor/Private/Modification/TransactionManager.cpp b/Engine/Source/Editor/CurveEditor/Private/Modification/TransactionManager.cpp
--- a/Engine/Source/Editor/CurveEditor/Private/Modification/TransactionManager.cpp
+++ b/Engine/Source/Editor/CurveEditor/Private/Modification/TransactionManager.cpp
@@ -20,6 +20,7 @@ void FTransactionManager::AppendChange(TUniquePtr<FCurveEditorCommandChange> InC
 	if (ensureAlways(GUndo) && ensure(DummyTransactionObject) && ensure(InChange))
 	{
 		GUndo->StoreUndo(DummyTransactionObject, MoveTemp(InChange));
+		OnCommandAppendedDelegate.Broadcast(*InChange);
 	}
 }
 
diff --git a/Engine/Source/Editor/CurveEditor/Private/SCurveEditorPanel.cpp b/Engine/Source/Editor/CurveEditor/Private/SCurveEditorPanel.cpp
--- a/Engine/Source/Editor/CurveEditor/Private/SCurveEditorPanel.cpp
+++ b/Engine/Source/Editor/CurveEditor/Private/SCurveEditorPanel.cpp
@@ -514,7 +514,6 @@ void SCurveEditorPanel::UpdateAxisSnapping()
 
 void SCurveEditorPanel::Tick(const FGeometry& AllottedGeometry, const double InCurrentTime, const float InDeltaTime)
 {
-	bool bWasRefreshed = false;
 	if (bNeedsRefresh || CachedActiveCurvesSerialNumber != CurveEditor->GetActiveCurvesSerialNumber())
 	{
 		RebuildCurveViews();
@@ -525,7 +524,6 @@ void SCurveEditorPanel::Tick(const FGeometry& AllottedGeometry, const double InC
 		}
 		bNeedsRefresh = false;
 		CachedActiveCurvesSerialNumber = CurveEditor->GetActiveCurvesSerialNumber();
-		bWasRefreshed = true;
 	}
 
 	UpdateCommonCurveInfo();
@@ -743,6 +741,8 @@ void SCurveEditorPanel::RebuildCurveViews()
 	{
 		CurveViewsContainer->AddView(View.ToSharedRef());
 	}
+
+	OnPostRebuildCurveViewsDelegate.Broadcast();
 }
 
 void SCurveEditorPanel::UpdateCommonCurveInfo()
diff --git a/Engine/Source/Editor/CurveEditor/Private/SCurveEditorView.cpp b/Engine/Source/Editor/CurveEditor/Private/SCurveEditorView.cpp
--- a/Engine/Source/Editor/CurveEditor/Private/SCurveEditorView.cpp
+++ b/Engine/Source/Editor/CurveEditor/Private/SCurveEditorView.cpp
@@ -151,7 +151,6 @@ FCurveEditorScreenSpace SCurveEditorView::GetCurveSpace(FCurveModelID CurveID) c
 	FCurveInfo CurveInfo = CurveInfoByID.FindRef(CurveID);
 
 	FTransform2d CurveTransform = CurveInfo.ViewToCurveTransform;
-	FScale2d     AxisScale(0.0, 0.0);
 
 	double InputMin = 0.0, InputMax = 1.0;
 	GetInputBounds(InputMin, InputMax);
@@ -172,6 +171,13 @@ FCurveEditorScreenSpace SCurveEditorView::GetCurveSpace(FCurveModelID CurveID) c
 	return FCurveEditorScreenSpace(GetCachedGeometry().GetLocalSize(), InputMin, InputMax, AxisOutputMin, AxisOutputMax).ToCurveSpace(CurveTransform);
 }
 
+FTransform2d SCurveEditorView::GetViewToCurveTransform(FCurveModelID CurveID) const
+{
+	FCurveInfo CurveInfo = CurveInfoByID.FindRef(CurveID);
+	FTransform2d CurveTransform = CurveInfo.ViewToCurveTransform;
+	return CurveTransform;
+}
+
 FCurveEditorScreenSpaceH SCurveEditorView::GetHorizontalAxisSpace(FCurveEditorViewAxisID ID) const
 {
 	double Min = 0.0;
diff --git a/Engine/Source/Editor/CurveEditor/Public/Modification/TransactionManager.h b/Engine/Source/Editor/CurveEditor/Public/Modification/TransactionManager.h
--- a/Engine/Source/Editor/CurveEditor/Public/Modification/TransactionManager.h
+++ b/Engine/Source/Editor/CurveEditor/Public/Modification/TransactionManager.h
@@ -30,6 +30,10 @@ public:
 	 */
 	void AppendChange(TUniquePtr<FCurveEditorCommandChange> InChange) const;
 
+	DECLARE_MULTICAST_DELEGATE_OneParam(FOnCommandAppended, const FCurveEditorCommandChange&);
+	/** Invoked after a command has been appended to the undo stack. */
+	FOnCommandAppended& OnCommandAppended() { return OnCommandAppendedDelegate; }
+
 	//~ Begin FGCObject Interface
 	virtual void AddReferencedObjects(FReferenceCollector& Collector) override;
 	virtual FString GetReferencerName() const override;
@@ -39,6 +43,9 @@ private:
 
 	/** The UObject that transactions are associated with. The transaction system needs to associate FCommandChanges with an UObject. */
 	TObjectPtr<UCurveEditorTransactionObject> DummyTransactionObject;
+
+	/** Invoked after a command has been appended to the undo stack. */
+	FOnCommandAppended OnCommandAppendedDelegate;
 };
 }
 
diff --git a/Engine/Source/Editor/CurveEditor/Public/SCurveEditorPanel.h b/Engine/Source/Editor/CurveEditor/Public/SCurveEditorPanel.h
--- a/Engine/Source/Editor/CurveEditor/Public/SCurveEditorPanel.h
+++ b/Engine/Source/Editor/CurveEditor/Public/SCurveEditorPanel.h
@@ -185,6 +185,9 @@ class CURVEEDITOR_API SCurveEditorPanel : public SCompoundWidget
 
 	/** Enable/disable pending focus */
 	void EnablePendingFocusOnHovering(const bool InEnabled);
+	
+	/** Broadcasts after the curve views have been rebuilt (by RebuildCurveViews). */
+	FSimpleMulticastDelegate& OnPostRebuildCurveViews() { return OnPostRebuildCurveViewsDelegate; }
 
 private:
 	// SWidget Interface
@@ -430,6 +433,9 @@ private:
 	/** The last set View Mode for this UI. */
 	ECurveEditorViewID DefaultViewID;
 
+	/** Broadcasts after the curve views have been rebuilt (by RebuildCurveViews). */
+	FSimpleMulticastDelegate OnPostRebuildCurveViewsDelegate;
+
 	TMultiMap<ECurveEditorViewID, TSharedRef<SCurveEditorView>> FreeViewsByType;
 
 	/** The scrool box that all curve views live inside */
diff --git a/Engine/Source/Editor/CurveEditor/Public/SCurveEditorView.h b/Engine/Source/Editor/CurveEditor/Public/SCurveEditorView.h
--- a/Engine/Source/Editor/CurveEditor/Public/SCurveEditorView.h
+++ b/Engine/Source/Editor/CurveEditor/Public/SCurveEditorView.h
@@ -140,6 +140,9 @@ public:
 	 */
 	FCurveEditorScreenSpace GetCurveSpace(FCurveModelID CurveID) const;
 
+	/** @return The transform used to translate view space (absolute key coordinates) to the curve space (e.g. 0-1 range in normalized view). */
+	FTransform2d GetViewToCurveTransform(FCurveModelID CurveID) const;
+
 	/**
 	 * Retrieve the horizontal screen-space information for the specified axis
 	 */
