Curve Editor: Lattice tool will no longer highlight UI elements while MMB dragging.

diff --git a/Engine/Plugins/Editor/CurveEditorTools/Source/CurveEditorTools/Private/Tools/Lattice/CurveEditorLatticeTool.cpp b/Engine/Plugins/Editor/CurveEditorTools/Source/CurveEditorTools/Private/Tools/Lattice/CurveEditorLatticeTool.cpp
--- a/Engine/Plugins/Editor/CurveEditorTools/Source/CurveEditorTools/Private/Tools/Lattice/CurveEditorLatticeTool.cpp
+++ b/Engine/Plugins/Editor/CurveEditorTools/Source/CurveEditorTools/Private/Tools/Lattice/CurveEditorLatticeTool.cpp
@@ -406,8 +406,9 @@ void FCurveEditorLatticeTool::RecomputeDrawData(const FGeometry& InViewContainer
 	}
 	
 	DrawData->ControlEdges = ComputeDrawControlEdges(DrawData->ControlPoints);
-	// No hover highlights while dragging
-	if (!DelayedDragData.IsSet())
+	// No hover highlights while dragging (or while the user is doing anything with mouse buttons down in the application).
+	const bool bAnyMouseButtonsDown = !FSlateApplication::Get().GetPressedMouseButtons().IsEmpty();
+	if (!DelayedDragData.IsSet() && !bAnyMouseButtonsDown)
 	{
 		DrawData->HoverState = ComputeLatticeHoverState(
 			InViewContainerGeometry, FinalMouseScreenSpacePosition,
