Curve Editor: Fixed dereferencing of nullptr after MoveTemp. This has caused no issues so far but it's bad behaviour.

diff --git a/Engine/Plugins/Editor/CurveEditorTools/Source/CurveEditorTools/Private/Tools/Lattice/CurveEditorLatticeTool.cpp b/Engine/Plugins/Editor/CurveEditorTools/Source/CurveEditorTools/Private/Tools/Lattice/CurveEditorLatticeTool.cpp
--- a/Engine/Plugins/Editor/CurveEditorTools/Source/CurveEditorTools/Private/Tools/Lattice/CurveEditorLatticeTool.cpp
+++ b/Engine/Plugins/Editor/CurveEditorTools/Source/CurveEditorTools/Private/Tools/Lattice/CurveEditorLatticeTool.cpp
@@ -354,7 +354,7 @@ void FCurveEditorLatticeTool::OnCurvesModified()
 	LatticeDeformerState->bHasRequestedRefresh = true;
 }
 
-void FCurveEditorLatticeTool::OnUndoCommandAppended(const CurveEditor::FCurveEditorCommandChange& InCommand)
+void FCurveEditorLatticeTool::OnUndoCommandAppended()
 {
 	if (const TSharedPtr<FCurveEditor> CurveEditor = WeakCurveEditor.Pin())
 	{
diff --git a/Engine/Plugins/Editor/CurveEditorTools/Source/CurveEditorTools/Private/Tools/Lattice/CurveEditorLatticeTool.h b/Engine/Plugins/Editor/CurveEditorTools/Source/CurveEditorTools/Private/Tools/Lattice/CurveEditorLatticeTool.h
--- a/Engine/Plugins/Editor/CurveEditorTools/Source/CurveEditorTools/Private/Tools/Lattice/CurveEditorLatticeTool.h
+++ b/Engine/Plugins/Editor/CurveEditorTools/Source/CurveEditorTools/Private/Tools/Lattice/CurveEditorLatticeTool.h
@@ -200,7 +200,7 @@ private:
 	/** Regenerates the lattice. */
 	void OnCurvePanelRebuilt() { RecomputeLatticeGrid(ELatticeUpdateFlags::NoTransaction); }
 	/** Handles selection transactional changes: updates the lattice shape change in the same transaction. */
-	void OnUndoCommandAppended(const CurveEditor::FCurveEditorCommandChange& InCommand);
+	void OnUndoCommandAppended();
 	/**  */
 	void OnTransactionStateChanged(const FTransactionContext& TransactionContext, ETransactionStateEventType TransactionStateEvent);
 	
diff --git a/Engine/Source/Editor/CurveEditor/Private/Modification/TransactionManager.cpp b/Engine/Source/Editor/CurveEditor/Private/Modification/TransactionManager.cpp
--- a/Engine/Source/Editor/CurveEditor/Private/Modification/TransactionManager.cpp
+++ b/Engine/Source/Editor/CurveEditor/Private/Modification/TransactionManager.cpp
@@ -23,7 +23,7 @@ void FTransactionManager::AppendChange(TUniquePtr<FCurveEditorCommandChange> InC
 		&& ensure(DummyTransactionObject) && ensure(InChange))
 	{
 		GUndo->StoreUndo(DummyTransactionObject, MoveTemp(InChange));
-		OnCommandAppendedDelegate.Broadcast(*InChange);
+		OnCommandAppendedDelegate.Broadcast();
 	}
 }
 
diff --git a/Engine/Source/Editor/CurveEditor/Public/Modification/TransactionManager.h b/Engine/Source/Editor/CurveEditor/Public/Modification/TransactionManager.h
--- a/Engine/Source/Editor/CurveEditor/Public/Modification/TransactionManager.h
+++ b/Engine/Source/Editor/CurveEditor/Public/Modification/TransactionManager.h
@@ -33,7 +33,7 @@ public:
 	 */
 	UE_API void AppendChange(TUniquePtr<FCurveEditorCommandChange> InChange, const FText& InDescription) const;
 
-	DECLARE_MULTICAST_DELEGATE_OneParam(FOnCommandAppended, const FCurveEditorCommandChange&);
+	DECLARE_MULTICAST_DELEGATE(FOnCommandAppended);
 	/** Invoked after a command has been appended to the undo stack. */
 	FOnCommandAppended& OnCommandAppended() { return OnCommandAppendedDelegate; }
 
