Curve Editor: Fix crash caused by stale keys in selection.
- Fixed transform tool scale not checking for nullptr

The proper fix is to fix the fact that keys in the selection are stale, which will be done in a follow up fix, see UE-32713.

diff --git a/Engine/Plugins/Editor/CurveEditorTools/Source/CurveEditorTools/Private/Tools/CurveEditorTransformTool.cpp b/Engine/Plugins/Editor/CurveEditorTools/Source/CurveEditorTools/Private/Tools/CurveEditorTransformTool.cpp
--- a/Engine/Plugins/Editor/CurveEditorTools/Source/CurveEditorTools/Private/Tools/CurveEditorTransformTool.cpp
+++ b/Engine/Plugins/Editor/CurveEditorTools/Source/CurveEditorTools/Private/Tools/CurveEditorTransformTool.cpp
@@ -1217,6 +1217,52 @@ void FCurveEditorTransformTool::ScaleFrom(const FVector2D& InPanelSpaceCenter, c
 	}
 }
 
+namespace UE::CurveEditorTools
+{
+static bool HaveSameScales(const FCurveEditor& CurveEditor)
+{
+	TArray<FCurveModelID> CurveModelIDs;
+	CurveEditor.GetSelection().GetAll().GetKeys(CurveModelIDs);
+	
+	if (CurveModelIDs.Num() == 0)
+	{
+		return false;
+	}
+
+	if (CurveModelIDs.Num() == 1)
+	{
+		return true;
+	}
+
+	// To handle FindFirstInteractiveView returning nullptr, we'll save the first FCurveEditorScreenSpace we encounter and compare against the rest.
+	bool bCurvesHaveSameScales = true;
+	TOptional<FCurveEditorScreenSpace> FirstCurveViewSpace;
+	
+	for (int32 ModelIndex = 0; ModelIndex < CurveModelIDs.Num(); ModelIndex++)
+	{
+		const SCurveEditorView* View = CurveEditor.FindFirstInteractiveView(CurveModelIDs[ModelIndex]);
+		if (!ensureMsgf(View, TEXT("No view found for curve. Probably the selection contains stale keys.")))
+		{
+			continue;
+		}
+		
+		const FCurveEditorScreenSpace CurveViewSpace = View->GetCurveSpace(CurveModelIDs[ModelIndex]);
+		if (FirstCurveViewSpace.IsSet())
+		{
+			bCurvesHaveSameScales &= FMath::IsNearlyEqual(CurveViewSpace.PixelsPerInput(), FirstCurveViewSpace->PixelsPerInput())
+				&& FMath::IsNearlyEqual(CurveViewSpace.PixelsPerOutput(), FirstCurveViewSpace->PixelsPerOutput());
+		}
+		else
+		{
+			FirstCurveViewSpace = CurveViewSpace;
+		}
+		
+	}
+
+	return bCurvesHaveSameScales;
+}
+}
+
 void FCurveEditorTransformTool::Tick(const FGeometry& AllottedGeometry, const double InCurrentTime, const float InDeltaTime)
 {
 	// We update the size and position of the box every frame as some scale operations aren't 1:1
@@ -1234,18 +1280,7 @@ void FCurveEditorTransformTool::Tick(const FGeometry& AllottedGeometry, const do
 	const FFrameRate TickResolution = CurveEditor->GetTimeSliderController() ? CurveEditor->GetTimeSliderController()->GetTickResolution() : FFrameRate();
 
 	// Check if all curve models have the same scales
-	TArray<FCurveModelID> CurveModelIDs;
-	CurveEditor->GetSelection().GetAll().GetKeys(CurveModelIDs);
-	bCurvesHaveSameScales = CurveModelIDs.Num() == 1;
-	if (CurveModelIDs.Num() > 1)
-	{
-		FCurveEditorScreenSpace FirstCurveViewSpace = CurveEditor->FindFirstInteractiveView(CurveModelIDs[0])->GetCurveSpace(CurveModelIDs[0]);
-		for (int32 ModelIndex = 1; ModelIndex < CurveModelIDs.Num(); ModelIndex++)
-		{
-			FCurveEditorScreenSpace CurveViewSpace = CurveEditor->FindFirstInteractiveView(CurveModelIDs[ModelIndex])->GetCurveSpace(CurveModelIDs[ModelIndex]);
-			bCurvesHaveSameScales |= FMath::IsNearlyEqual(CurveViewSpace.PixelsPerInput(), FirstCurveViewSpace.PixelsPerInput()) && FMath::IsNearlyEqual(CurveViewSpace.PixelsPerOutput(), FirstCurveViewSpace.PixelsPerOutput());
-		}
-	}
+	bCurvesHaveSameScales = UE::CurveEditorTools::HaveSameScales(*CurveEditor);
 
 	// remove tool options widget if less than two selected or there are two different curves
 	if ((CurveEditor->GetSelection().Count() <= 1 || !bCurvesHaveSameScales) 
@@ -1368,5 +1403,4 @@ double FCurveEditorTransformTool::ModifyWeightByInterpType(double Value) const
 	}
 	return Result;
 }
-#undef LOCTEXT_NAMESPACE // "CurveEditorToolCommands"
-
+#undef LOCTEXT_NAMESPACE // "CurveEditorToolCommands"
\ No newline at end of file
