Curve Editor: Fix multi-scale tool not getting rid of stacked keys.

This would cause undo / redo operations to delete the stack keys by accident. The user would not be able to undo losing those keys.

diff --git a/Engine/Plugins/Editor/CurveEditorTools/Source/CurveEditorTools/Private/Tools/CurveEditorMultiScaleTool.cpp b/Engine/Plugins/Editor/CurveEditorTools/Source/CurveEditorTools/Private/Tools/CurveEditorMultiScaleTool.cpp
--- a/Engine/Plugins/Editor/CurveEditorTools/Source/CurveEditorTools/Private/Tools/CurveEditorMultiScaleTool.cpp
+++ b/Engine/Plugins/Editor/CurveEditorTools/Source/CurveEditorTools/Private/Tools/CurveEditorMultiScaleTool.cpp
@@ -471,10 +471,13 @@ void FCurveEditorMultiScaleTool::OnDrag(const FPointerEvent& InMouseEvent, const
 
 void FCurveEditorMultiScaleTool::OnDragEnd()
 {
-	// This finalizes the transaction
+	// Call SetKeyPositions again with EPropertyChangeType::ValueSet to trigger any follow up operations, like getting rid of stacked keys...
+	ApplyLastKeyPositions();
+	// ... then finalize the transaction
 	ActiveTransaction.Reset();
+	
+	KeysByCurve.Reset(); // User might not use our tool anymore - clean-up any cached memory.
 	DragDelta = TOptional<FVector2D>();
-
 	ToolOptions.XScale = 1.0f;
 	ToolOptions.YScale = 1.0f;
 }
@@ -603,11 +606,9 @@ void FCurveEditorMultiScaleTool::ScaleUnique(const FVector2D& InChangeAmount, co
 	// this generically by scaling coordinates relative to an arbitrary center
 	const FVector2D PercentChanged = FVector2D(1.0f, 1.0f) + InChangeAmount; // ie: 5 pixel change on a 100 wide gives you 1.05
 	
-	TArray<FKeyPosition> NewKeyPositionScratch;
-	
 	// We now know if we need to affect both X and Y, and we know where we're scaling from. Now we can loop through the keys and actually modify their positions.
 	// We perform the scale on both axis (for simplicity) and then read which axis it should effect before assigning it back to the key position.
-	for (const FKeyData& KeyData : KeysByCurve)
+	for (FKeyData& KeyData : KeysByCurve)
 	{
 		const SCurveEditorView* View = CurveEditor->FindFirstInteractiveView(KeyData.CurveID);
 		if (!View)
@@ -621,8 +622,8 @@ void FCurveEditorMultiScaleTool::ScaleUnique(const FVector2D& InChangeAmount, co
 		const double CurveSpaceCenterInput = KeyData.Pivot.X;
 		const double CurveSpaceCenterOutput = KeyData.Pivot.Y;
 
-		NewKeyPositionScratch.Reset();
-		NewKeyPositionScratch.Reserve(KeyData.StartKeyPositions.Num());
+		TArray<FKeyPosition>& LastDraggedKeyPositions = KeyData.LastDraggedKeyPositions;
+		LastDraggedKeyPositions.Empty(KeyData.StartKeyPositions.Num());
 
 		FCurveSnapMetrics SnapMetrics = CurveEditor->GetCurveSnapMetrics(KeyData.CurveID);
 
@@ -647,10 +648,10 @@ void FCurveEditorMultiScaleTool::ScaleUnique(const FVector2D& InChangeAmount, co
 				StartPosition.OutputValue = View->IsValueSnapEnabled() ? SnapMetrics.SnapOutput(NewOutput) : NewOutput;
 			}
 
-			NewKeyPositionScratch.Add(StartPosition);
+			LastDraggedKeyPositions.Add(StartPosition);
 		}
 
-		CurveModel->SetKeyPositions(KeyData.Handles, NewKeyPositionScratch);
+		CurveModel->SetKeyPositions(KeyData.Handles, LastDraggedKeyPositions, EPropertyChangeType::Interactive);
 	}
 	
 	// This recalculates the user specified tangents based on the distance moved. We'll flatten the curves depending on the resulting selection's height.
@@ -661,6 +662,23 @@ void FCurveEditorMultiScaleTool::ScaleUnique(const FVector2D& InChangeAmount, co
 	}
 }
 
+void FCurveEditorMultiScaleTool::ApplyLastKeyPositions()
+{
+	const TSharedPtr<FCurveEditor> CurveEditor = WeakCurveEditor.Pin();
+	if (!CurveEditor.IsValid())
+	{
+		return;
+	}
+	
+	for (const FKeyData& KeyData : KeysByCurve)
+	{
+		if (FCurveModel* Curve = CurveEditor->FindCurve(KeyData.CurveID))
+		{
+			Curve->SetKeyPositions(KeyData.Handles, KeyData.LastDraggedKeyPositions, EPropertyChangeType::ValueSet);
+		}
+	}
+}
+
 FVector2D FCurveEditorMultiScaleTool::GetPivot(FCurveModel* InCurve, const TArray<FKeyPosition>& InKeyPositions) const
 {
 	FVector2D Pivot = FVector2D::ZeroVector;
diff --git a/Engine/Plugins/Editor/CurveEditorTools/Source/CurveEditorTools/Private/Tools/CurveEditorMultiScaleTool.h b/Engine/Plugins/Editor/CurveEditorTools/Source/CurveEditorTools/Private/Tools/CurveEditorMultiScaleTool.h
--- a/Engine/Plugins/Editor/CurveEditorTools/Source/CurveEditorTools/Private/Tools/CurveEditorMultiScaleTool.h
+++ b/Engine/Plugins/Editor/CurveEditorTools/Source/CurveEditorTools/Private/Tools/CurveEditorMultiScaleTool.h
@@ -127,6 +127,16 @@ private:
 	void UpdateBoundingBox();
 
 	void ScaleUnique(const FVector2D& InChangeAmount, const bool bInAffectsX, const bool bInAffectsY);
+	
+	/**
+	 * Calls SetKeyPositions with FKeyData::LastDraggedKeyPositions and EPropertyChangeType::ValueSet tro trigger any final operations.
+	 * 
+	 * Some FCurveModels, such as Sequencer's FCurveChannelModel, perform special operations with EPropertyChangeType::ValueSet, e.g. cleansing of
+	 * stacked values.
+	 * 
+	 * Called at end of drag.
+	 */
+	void ApplyLastKeyPositions();
 
 private:
 	/** Weak pointer back to the owning curve editor. */
@@ -154,6 +164,8 @@ private:
 		TArray<FKeyHandle> Handles;
 		/** The extended key info for each of the above handles */
 		TArray<FKeyPosition> StartKeyPositions;
+		/** Used in OnEndDrag to send final key updates */
+		TArray<FKeyPosition> LastDraggedKeyPositions;
 		/** Pivot to scale from */
 		FVector2D Pivot;
 	};
