Curve Editor: Fix undoing lattice, transform, and multi-scale tool after they stacked keys would skip re-selecting the keys that were removed.

diff --git a/Engine/Plugins/Editor/CurveEditorTools/Source/CurveEditorTools/Private/Tools/Lattice/CurveEditorLatticeTool.cpp b/Engine/Plugins/Editor/CurveEditorTools/Source/CurveEditorTools/Private/Tools/Lattice/CurveEditorLatticeTool.cpp
--- a/Engine/Plugins/Editor/CurveEditorTools/Source/CurveEditorTools/Private/Tools/Lattice/CurveEditorLatticeTool.cpp
+++ b/Engine/Plugins/Editor/CurveEditorTools/Source/CurveEditorTools/Private/Tools/Lattice/CurveEditorLatticeTool.cpp
@@ -234,12 +234,6 @@ void FCurveEditorLatticeTool::OnFinishedPointerInput()
 
 void FCurveEditorLatticeTool::OnToolActivated()
 {
-	if (const TSharedPtr<FCurveEditor> CurveEditor = WeakCurveEditor.Pin())
-	{
-		UpdateDeformerBounds(CurveEditor.ToSharedRef());
-		CurveEditor->GetTransactionManager()->OnCommandAppended().AddRaw(this, &FCurveEditorLatticeTool::OnUndoCommandAppended);
-	}
-	
 	if (UTransBuffer* Transactor = GEditor ? Cast<UTransBuffer>(GEditor->Trans) : nullptr)
 	{
 		Transactor->OnTransactionStateChanged().AddRaw(this, &FCurveEditorLatticeTool::OnTransactionStateChanged);
@@ -354,15 +348,6 @@ void FCurveEditorLatticeTool::OnCurvesModified()
 	LatticeDeformerState->bHasRequestedRefresh = true;
 }
 
-void FCurveEditorLatticeTool::OnUndoCommandAppended()
-{
-	if (const TSharedPtr<FCurveEditor> CurveEditor = WeakCurveEditor.Pin())
-	{
-		// This will diff the selection serial number and modify the undo object.
-		UpdateDeformerBounds(CurveEditor.ToSharedRef(), ELatticeUpdateFlags::NoTransaction);
-	}
-}
-
 void FCurveEditorLatticeTool::OnTransactionStateChanged(
 	const FTransactionContext& Context, ETransactionStateEventType TransactionStateEvent
 	)
@@ -873,11 +858,6 @@ bool FCurveEditorLatticeTool::ComputeHoverStateForLeftClick(
 	
 void FCurveEditorLatticeTool::UnregisterDelegates() const
 {
-	if (const TSharedPtr<FCurveEditor> CurveEditor = WeakCurveEditor.Pin())
-	{
-		CurveEditor->GetTransactionManager()->OnCommandAppended().RemoveAll(this);
-	}
-	
 	if (UTransBuffer* Transactor = GEditor ? Cast<UTransBuffer>(GEditor->Trans) : nullptr)
 	{
 		Transactor->OnTransactionStateChanged().RemoveAll(this);
diff --git a/Engine/Plugins/Editor/CurveEditorTools/Source/CurveEditorTools/Private/Tools/Lattice/CurveEditorLatticeTool.h b/Engine/Plugins/Editor/CurveEditorTools/Source/CurveEditorTools/Private/Tools/Lattice/CurveEditorLatticeTool.h
--- a/Engine/Plugins/Editor/CurveEditorTools/Source/CurveEditorTools/Private/Tools/Lattice/CurveEditorLatticeTool.h
+++ b/Engine/Plugins/Editor/CurveEditorTools/Source/CurveEditorTools/Private/Tools/Lattice/CurveEditorLatticeTool.h
@@ -202,8 +202,6 @@ private:
 	void OnCurvesModified();
 	/** Regenerates the lattice. */
 	void OnCurvePanelRebuilt() { RecomputeLatticeGrid(ELatticeUpdateFlags::NoTransaction); }
-	/** Handles selection transactional changes: updates the lattice shape change in the same transaction. */
-	void OnUndoCommandAppended();
 	/**  */
 	void OnTransactionStateChanged(const FTransactionContext& TransactionContext, ETransactionStateEventType TransactionStateEvent);
 	
diff --git a/Engine/Source/Editor/CurveEditor/Private/Selection/SelectionCleanser.cpp b/Engine/Source/Editor/CurveEditor/Private/Selection/SelectionCleanser.cpp
--- a/Engine/Source/Editor/CurveEditor/Private/Selection/SelectionCleanser.cpp
+++ b/Engine/Source/Editor/CurveEditor/Private/Selection/SelectionCleanser.cpp
@@ -6,6 +6,7 @@
 #include "Editor.h"
 #include "SCurveEditor.h"
 #include "SelectionUtils.h"
+#include "Modification/Utils/ScopedSelectionChange.h"
 
 namespace UE::CurveEditor
 {
@@ -64,7 +65,8 @@ void FSelectionCleanser::OnCurvesChanged(const FGenericCurveChangeData& InChange
 	{
 		return;
 	}
-	
+
+	const FScopedSelectionChange SelectionChange(CurveEditorPin);
 	for (const TPair<FCurveModelID, FCurveKeyData>& Pair : InChange.RemoveKeysData.SavedCurveState)
 	{
 		CleanseSelection(CurveEditorPin.ToSharedRef(), CurveEditorPin->Selection, MakeConstArrayView({ Pair.Key }));
