Curve Editor: Retime anchor size adjustment

diff --git a/Engine/Plugins/Editor/CurveEditorTools/Source/CurveEditorTools/Private/Tools/Retime/RetimeAnchor.cpp b/Engine/Plugins/Editor/CurveEditorTools/Source/CurveEditorTools/Private/Tools/Retime/RetimeAnchor.cpp
--- a/Engine/Plugins/Editor/CurveEditorTools/Source/CurveEditorTools/Private/Tools/Retime/RetimeAnchor.cpp
+++ b/Engine/Plugins/Editor/CurveEditorTools/Source/CurveEditorTools/Private/Tools/Retime/RetimeAnchor.cpp
@@ -14,10 +14,12 @@ static FGeometry GetRetimeBarGeometry(const FGeometry& InWidgetGeometry, const F
 	const double AnchorPosition = HorizontalTransform.SecondsToScreen(ValueInSeconds);
 	const double BarHeight = InWidgetGeometry.GetLocalSize().Y - (CloseButtonWidth + CloseButtonPadding);
 
-	FVector2D BarRetimeSize = FVector2D(AnchorRetimeBarWidth + ExtraHorizontalPadding, BarHeight-(2*AnchorRetimeBarVerticalOffset));
+	// Set the size of the retime bar as half the height of the move-only bar
+	FVector2D BarRetimeSize = FVector2D(AnchorRetimeBarWidth + ExtraHorizontalPadding, BarHeight/2);
 	// Horizontally center the retime bar on top of the move-only bar
 	const double HorizontalOffset = (AnchorMoveOnlyBarWidth - ExtraHorizontalPadding - AnchorRetimeBarWidth)/2;
-	const FVector2D BarRetimeOffset = FVector2D(AnchorPosition + HorizontalOffset, AnchorRetimeBarVerticalOffset);
+	// Vertically center the retime bar - offset is quarter of the height of the move-only bar (H/2)/2 = H/4
+	const FVector2D BarRetimeOffset = FVector2D(AnchorPosition + HorizontalOffset, BarHeight/4);
 
 	return InWidgetGeometry.MakeChild(BarRetimeSize, FSlateLayoutTransform(BarRetimeOffset));
 }
diff --git a/Engine/Plugins/Editor/CurveEditorTools/Source/CurveEditorTools/Private/Tools/Retime/RetimeAnchor.h b/Engine/Plugins/Editor/CurveEditorTools/Source/CurveEditorTools/Private/Tools/Retime/RetimeAnchor.h
--- a/Engine/Plugins/Editor/CurveEditorTools/Source/CurveEditorTools/Private/Tools/Retime/RetimeAnchor.h
+++ b/Engine/Plugins/Editor/CurveEditorTools/Source/CurveEditorTools/Private/Tools/Retime/RetimeAnchor.h
@@ -50,7 +50,6 @@ namespace UE::CurveEditorTools
 {
 	constexpr double AnchorMoveOnlyBarWidth = 3.0;
 	constexpr double AnchorRetimeBarWidth = 5.0;
-	constexpr double AnchorRetimeBarVerticalOffset = 100;
 	constexpr double CloseButtonWidth = 18.0;
 	constexpr double CloseButtonPadding = 2.0;
 	constexpr double HighlightOpacity = 0.75;
