Curve Editor: Lattice tool now restores shape when you undo.

diff --git a/Engine/Plugins/Editor/CurveEditorTools/Source/CurveEditorTools/Private/Tools/Lattice/CurveEditorLatticeTool.cpp b/Engine/Plugins/Editor/CurveEditorTools/Source/CurveEditorTools/Private/Tools/Lattice/CurveEditorLatticeTool.cpp
--- a/Engine/Plugins/Editor/CurveEditorTools/Source/CurveEditorTools/Private/Tools/Lattice/CurveEditorLatticeTool.cpp
+++ b/Engine/Plugins/Editor/CurveEditorTools/Source/CurveEditorTools/Private/Tools/Lattice/CurveEditorLatticeTool.cpp
@@ -46,7 +46,8 @@ FCurveEditorLatticeTool::FCurveEditorLatticeTool(TWeakPtr<FCurveEditor> InCurveE
 		// LatticeDeformerState is unset if the user redo-es the first transaction which set the initial selection...
 		if (LatticeDeformerState)
 		{
-			ApplyUndoState(*Editor);
+			LatticeDeformerState->bHasRequestedRefresh = true;
+			LatticeDeformerState->RefreshFlags |= ELatticeUpdateFlags::TryRestoreBoundsFromUndo;
 			// If user undo-es the first transaction, next tick we'll notice that nothing is selected and destroy the deformer.
 		}
 	});
@@ -308,7 +309,7 @@ TArray<FVector2D> FCurveEditorLatticeTool::FLatticeDeformerState::TransformContr
 	return Result;
 }
 
-void FCurveEditorLatticeTool::UpdateDeformerBounds(const TSharedRef<FCurveEditor>& InCurveEditor, ELatticeUpdateFlags InRefreshFlags)
+void FCurveEditorLatticeTool::UpdateDeformerBounds(const TSharedRef<FCurveEditor>& InCurveEditor)
 {
 	if (SelectionSerialNumber != InCurveEditor->Selection.GetSerialNumber())
 	{
@@ -337,19 +338,10 @@ void FCurveEditorLatticeTool::UpdateDeformerBounds(const TSharedRef<FCurveEditor
 void FCurveEditorLatticeTool::OnCurvesModified()
 {
 	// The modification was driven by us?
-	if (LatticeDeformerState->bIsModifyingCurves)
+	if (!LatticeDeformerState->bIsModifyingCurves)
 	{
-		return;
-	}
-
-	// If an external change happens during a transaction or undo / redo, and we refresh the grid shape next tick, do not generate another transaction.
-	// Instead of next tick, we should defer the change to the end of the current transaction but that is a bit tricky.
-	// Instead of implementing the deferring TODO UE-242452: the lattice tool should stop tracking key selection and we should make core Curve Editor track the key selection with transactions. 
-	if (GUndo || GIsTransacting)
-	{
-		LatticeDeformerState->RefreshFlags = ELatticeUpdateFlags::NoTransaction;
+		LatticeDeformerState->bHasRequestedRefresh = true;
 	}
-	LatticeDeformerState->bHasRequestedRefresh = true;
 }
 
 void FCurveEditorLatticeTool::OnTransactionStateChanged(
@@ -384,11 +376,17 @@ void FCurveEditorLatticeTool::OnTransactionStateChanged(
 
 void FCurveEditorLatticeTool::RecomputeLatticeGrid(const TSharedRef<FCurveEditor>& InCurveEditor, ELatticeUpdateFlags InRefreshFlags)
 {
+	if (EnumHasAnyFlags(InRefreshFlags, ELatticeUpdateFlags::TryRestoreBoundsFromUndo)
+		&& ApplyUndoState(*InCurveEditor))
+	{
+		return;
+	}
+	
 	// This will calculate a new rectangular grid around the points.
 	// In the past, we tried to preserve the existing shape by just moving it over the moved keys but there were too many different tools interacting
 	// which made this difficult to get right in every case. Hence, we decided just to recompute the grid on move.
 	LatticeDeformerState.Reset();
-	UpdateDeformerBounds(InCurveEditor, InRefreshFlags);
+	UpdateDeformerBounds(InCurveEditor);
 }
 	
 void FCurveEditorLatticeTool::RecomputeDrawData(const FGeometry& InViewContainerGeometry, const TSharedPtr<FCurveEditor>& InCurveEditor)
@@ -880,8 +878,13 @@ void FCurveEditorLatticeTool::SaveUndoState(const FCurveEditor& InCurveEditor)
 	}
 }
 
-void FCurveEditorLatticeTool::ApplyUndoState(FCurveEditor& InCurveEditor)
+bool FCurveEditorLatticeTool::ApplyUndoState(FCurveEditor& InCurveEditor)
 {
+	if (!LatticeDeformerState)
+	{
+		return false;
+	}
+	
 	const bool bWereViewModesCompatible = ApplySnapshot(
 		UndoObject->Snapshot, InCurveEditor,
 		LatticeDeformerState->ControlPointToCurveSpace,
@@ -894,8 +897,10 @@ void FCurveEditorLatticeTool::ApplyUndoState(FCurveEditor& InCurveEditor)
 	// Since the lattice shapes are incompatible between modes, we'll just regenerate the entire grid.
 	if (!bWereViewModesCompatible)
 	{
-		RecomputeLatticeGrid(InCurveEditor.AsShared());
+		LatticeDeformerState.Reset();
 	}
+
+	return bWereViewModesCompatible;
 }
 }
 
diff --git a/Engine/Plugins/Editor/CurveEditorTools/Source/CurveEditorTools/Private/Tools/Lattice/CurveEditorLatticeTool.h b/Engine/Plugins/Editor/CurveEditorTools/Source/CurveEditorTools/Private/Tools/Lattice/CurveEditorLatticeTool.h
--- a/Engine/Plugins/Editor/CurveEditorTools/Source/CurveEditorTools/Private/Tools/Lattice/CurveEditorLatticeTool.h
+++ b/Engine/Plugins/Editor/CurveEditorTools/Source/CurveEditorTools/Private/Tools/Lattice/CurveEditorLatticeTool.h
@@ -33,8 +33,17 @@ struct FLatticeBounds;
 enum class ELatticeUpdateFlags : uint8
 {
 	None,
-	/** Don't generate any transaction for this update. */
-	NoTransaction = 1 << 1
+	
+	/**
+	 * Tries to restore the grid bounds from what is saved in the undo object.
+	 * This is set when an undo operation is performed.
+	 *
+	 * If the data is compatible, the grid is restored to that and the grid is not recomputed.
+	 * If the data is incompatible, the entire deformer is recomputed.
+	 *
+	 * The data may be incompatible, if e.g. you switched from normalized view to absolute view.
+	 */
+	TryRestoreBoundsFromUndo = 1 << 1,
 };
 ENUM_CLASS_FLAGS(ELatticeUpdateFlags);
 	
@@ -196,12 +205,12 @@ private:
 	TObjectPtr<UCurveEditorTools_LatticeUndoObject> UndoObject;
 
 	/** Creates or destroys LatticeDeformer based on the selection. */
-	void UpdateDeformerBounds(const TSharedRef<FCurveEditor>& InCurveEditor, ELatticeUpdateFlags InRefreshFlags = ELatticeUpdateFlags::None);
+	void UpdateDeformerBounds(const TSharedRef<FCurveEditor>& InCurveEditor);
 
 	/** Updates the bounds after the user has externally moved the keys, e.g. MMB move keys, etc. */
 	void OnCurvesModified();
 	/** Regenerates the lattice. */
-	void OnCurvePanelRebuilt() { RecomputeLatticeGrid(ELatticeUpdateFlags::NoTransaction); }
+	void OnCurvePanelRebuilt() { RecomputeLatticeGrid(); }
 	/**  */
 	void OnTransactionStateChanged(const FTransactionContext& TransactionContext, ETransactionStateEventType TransactionStateEvent);
 	
@@ -260,7 +269,7 @@ private:
 
 	/** Saves the undo state in the undo object. */
 	void SaveUndoState(const FCurveEditor& InCurveEditor);
-	void ApplyUndoState(FCurveEditor& InCurveEditor);
+	bool ApplyUndoState(FCurveEditor& InCurveEditor);
 };
 
 inline void FCurveEditorLatticeTool::RecomputeLatticeGrid(ELatticeUpdateFlags InRefreshFlags)
