Curve Editor: Fix lattice tool resetting its shape when you cause key stacking to remove keys

diff --git a/Engine/Plugins/Editor/CurveEditorTools/Source/CurveEditorTools/Private/Tools/Lattice/CurveEditorLatticeTool.cpp b/Engine/Plugins/Editor/CurveEditorTools/Source/CurveEditorTools/Private/Tools/Lattice/CurveEditorLatticeTool.cpp
--- a/Engine/Plugins/Editor/CurveEditorTools/Source/CurveEditorTools/Private/Tools/Lattice/CurveEditorLatticeTool.cpp
+++ b/Engine/Plugins/Editor/CurveEditorTools/Source/CurveEditorTools/Private/Tools/Lattice/CurveEditorLatticeTool.cpp
@@ -243,6 +243,11 @@ void FCurveEditorLatticeTool::OnToolActivated()
 	{
 		Transactor->OnTransactionStateChanged().AddRaw(this, &FCurveEditorLatticeTool::OnTransactionStateChanged);
 	}
+
+	if (const TSharedPtr<FCurveEditor> CurveEditorPin = WeakCurveEditor.Pin())
+	{
+		CurveEditorPin->Selection.OnSelectionChanged().AddRaw(this, &FCurveEditorLatticeTool::OnSelectionChanged);
+	}
 }
 
 void FCurveEditorLatticeTool::OnToolDeactivated()
@@ -311,12 +316,6 @@ TArray<FVector2D> FCurveEditorLatticeTool::FLatticeDeformerState::TransformContr
 
 void FCurveEditorLatticeTool::UpdateDeformerBounds(const TSharedRef<FCurveEditor>& InCurveEditor)
 {
-	if (SelectionSerialNumber != InCurveEditor->Selection.GetSerialNumber())
-	{
-		CancelDrag();
-		LatticeDeformerState.Reset();
-	}
-
 	if (LatticeDeformerState.IsSet())
 	{
 		return;
@@ -327,7 +326,6 @@ void FCurveEditorLatticeTool::UpdateDeformerBounds(const TSharedRef<FCurveEditor
 	const FLatticeBounds Bounds = ComputeBounds(*InCurveEditor);
 	if (Bounds.bIsVisible && !IsLatticeTooSmall(Bounds.MinValues, Bounds.MaxValues) && InCurveEditor->Selection.Count() > 1)
 	{
-		SelectionSerialNumber = InCurveEditor->Selection.GetSerialNumber();
 		LatticeDeformerState.Emplace(Bounds, InCurveEditor);
 		LatticeDeformerState->CurveChangeListener.OnCurveModified().AddRaw(this, &FCurveEditorLatticeTool::OnCurvesModified);
 		LatticeDeformerState->PanelRebuildListener.OnCurveViewRebuilt().AddRaw(this, &FCurveEditorLatticeTool::OnCurvePanelRebuilt);
@@ -344,6 +342,31 @@ void FCurveEditorLatticeTool::OnCurvesModified()
 	}
 }
 
+void FCurveEditorLatticeTool::OnSelectionChanged()
+{
+	const TSharedPtr<FCurveEditor> CurveEditorPin = WeakCurveEditor.Pin();
+	if (!CurveEditorPin)
+	{
+		return;
+	}
+
+	// A change in selection can be caused by the lattice tool: if the user stacked some keys, the stacked keys are removed. That causes a selection change.
+	// In that case, retain the shape of the lattice. Just save the current shape into the undo buffer so it can be restored.
+	if (LatticeDeformerState && LatticeDeformerState->bIsModifyingCurves)
+	{
+		SaveUndoState(*CurveEditorPin);
+	}
+	else if (GIsTransacting && LatticeDeformerState)
+	{
+		LatticeDeformerState->bHasRequestedRefresh = true;
+	}
+	else
+	{
+		CancelDrag();
+		RecomputeLatticeGrid(CurveEditorPin.ToSharedRef());
+	}
+}
+
 void FCurveEditorLatticeTool::OnTransactionStateChanged(
 	const FTransactionContext& Context, ETransactionStateEventType TransactionStateEvent
 	)
@@ -864,6 +887,11 @@ void FCurveEditorLatticeTool::UnregisterDelegates() const
 	{
 		Transactor->OnTransactionStateChanged().RemoveAll(this);
 	}
+
+	if (const TSharedPtr<FCurveEditor> CurveEditorPin = WeakCurveEditor.Pin())
+	{
+		CurveEditorPin->Selection.OnSelectionChanged().RemoveAll(this);
+	}
 }
 
 void FCurveEditorLatticeTool::SaveUndoState(const FCurveEditor& InCurveEditor)
@@ -888,9 +916,7 @@ bool FCurveEditorLatticeTool::ApplyUndoState(FCurveEditor& InCurveEditor)
 	const bool bWereViewModesCompatible = ApplySnapshot(
 		UndoObject->Snapshot, InCurveEditor,
 		LatticeDeformerState->ControlPointToCurveSpace,
-		LatticeDeformerState->GlobalDeformer, LatticeDeformerState->PerCurveData,
-		// We need to overwrite our serial number, or else in the next frame UpdateDeformerBounds will think the selection has changed and reinit.
-		SelectionSerialNumber
+		LatticeDeformerState->GlobalDeformer, LatticeDeformerState->PerCurveData
 		);
 
 	// The snapshot could e.g. have been taken in Absolute view mode, but we might currently be in Normalized view.
diff --git a/Engine/Plugins/Editor/CurveEditorTools/Source/CurveEditorTools/Private/Tools/Lattice/CurveEditorLatticeTool.h b/Engine/Plugins/Editor/CurveEditorTools/Source/CurveEditorTools/Private/Tools/Lattice/CurveEditorLatticeTool.h
--- a/Engine/Plugins/Editor/CurveEditorTools/Source/CurveEditorTools/Private/Tools/Lattice/CurveEditorLatticeTool.h
+++ b/Engine/Plugins/Editor/CurveEditorTools/Source/CurveEditorTools/Private/Tools/Lattice/CurveEditorLatticeTool.h
@@ -133,8 +133,6 @@ private:
 	};
 	/** Set when there is a valid selection in view. */
 	TOptional<FLatticeDeformerState> LatticeDeformerState;
-	/** The selection serial number when the deformer was initialized. */
-	uint32 SelectionSerialNumber = 0;
 	
 	/** How to draw LatticeDeformer. Set only when the selection, over which the controls are drawn, is visible.*/
 	TOptional<FLatticeControlsDrawData> DrawData;
@@ -209,6 +207,8 @@ private:
 
 	/** Updates the bounds after the user has externally moved the keys, e.g. MMB move keys, etc. */
 	void OnCurvesModified();
+	/** Skips recomputing the bounds if the selection was changed due us modifying curves (e.g. because of stacked keys being removed) */
+	void OnSelectionChanged();
 	/** Regenerates the lattice. */
 	void OnCurvePanelRebuilt() { RecomputeLatticeGrid(); }
 	/**  */
diff --git a/Engine/Plugins/Editor/CurveEditorTools/Source/CurveEditorTools/Private/Tools/Lattice/Undo/LatticeSnapshot.cpp b/Engine/Plugins/Editor/CurveEditorTools/Source/CurveEditorTools/Private/Tools/Lattice/Undo/LatticeSnapshot.cpp
--- a/Engine/Plugins/Editor/CurveEditorTools/Source/CurveEditorTools/Private/Tools/Lattice/Undo/LatticeSnapshot.cpp
+++ b/Engine/Plugins/Editor/CurveEditorTools/Source/CurveEditorTools/Private/Tools/Lattice/Undo/LatticeSnapshot.cpp
@@ -53,7 +53,7 @@ FArchive& operator<<(FArchive& InArchive, FLatticeSnapshot& InSnapshot)
 bool ApplySnapshot(
 	const FLatticeSnapshot& InSnapshot, const FCurveEditor& InCurveEditor,
 	FTransform2d& OutControlPointToCurveSpace, FGlobalLatticeDeformer2D& OutGlobalDeformer,
-	TMap<FCurveModelID, FPerCurveLatticeData>& OutPerCurveData, uint32& OutSerialNumber
+	TMap<FCurveModelID, FPerCurveLatticeData>& OutPerCurveData
 )
 {
 	if (!InSnapshot.CanApplySnapshot(InCurveEditor)
@@ -104,8 +104,6 @@ bool ApplySnapshot(
     	
     	OutPerCurveData.Add(Pair.Key, MoveTemp(Data));
     }
-
-	OutSerialNumber = InSnapshot.SelectionSerialNumber;
 	
 	return true;
 }
diff --git a/Engine/Plugins/Editor/CurveEditorTools/Source/CurveEditorTools/Private/Tools/Lattice/Undo/LatticeSnapshot.h b/Engine/Plugins/Editor/CurveEditorTools/Source/CurveEditorTools/Private/Tools/Lattice/Undo/LatticeSnapshot.h
--- a/Engine/Plugins/Editor/CurveEditorTools/Source/CurveEditorTools/Private/Tools/Lattice/Undo/LatticeSnapshot.h
+++ b/Engine/Plugins/Editor/CurveEditorTools/Source/CurveEditorTools/Private/Tools/Lattice/Undo/LatticeSnapshot.h
@@ -48,7 +48,7 @@ struct FLatticeSnapshot
 bool ApplySnapshot(
 	const FLatticeSnapshot& InSnapshot, const FCurveEditor& InCurveEditor,
 	FTransform2d& OutControlPointToCurveSpace, FGlobalLatticeDeformer2D& OutGlobalDeformer,
-	TMap<FCurveModelID, FPerCurveLatticeData>& OutPerCurveData, uint32& OutSerialNumber
+	TMap<FCurveModelID, FPerCurveLatticeData>& OutPerCurveData
 	);
 
 /** Takes a snapshot of the lattice tool and the curve editor's data required in the future to validate that the data can be applied. */
