Curve Editor: Fix retime tool freezing the editor when moving many keys:
- Instead of calling SetKeyPositions once for each key, all changes are now accumulated into one large array so SetKeyPositions is called only once
- The keys are processed using ParallelFor further reducing time

=== ParallelFor
Some key processing times for before and after using ParallelFor sorted by key number: Before min   | After min
- SetKeyPositions:  7,000 keys (9 curves)       13 ms    |   7 ms
- SetKeyPositions: 16,000 keys (21 curves)    66 ms    |  47 ms
- SetKeyPositions:150,000 keys                     357 ms    | 154 ms

=== Accumulate SetKeyPosition
Set all keys per curve at once (instead of individually) to improve performance.
Some stats for different key counts - times vary depending on number of keys and degree of change
(i.e. how far you move your mouse in quick succession).
Sorted by key number:                                  Before min   | Before max   | After min   | After max
- SetKeyPositions:  7,000 keys (9 curves)       90 ms    |   130 ms     |    7 ms     |   49 ms
- SetKeyPositions: 16,000 keys (21 curves)    25 s*    |    25 s*     |   47 ms     |  176 ms
- SetKeyPositions:150,000 keys                      33 s*    |    33 s*     |  154 ms     |    1 s
* retime not executed correctly - anchor not moved to mouse, drag "cancelled" -> OnDragEnd instantly called after first function call

diff --git a/Engine/Plugins/Editor/CurveEditorTools/Source/CurveEditorTools/Private/Tools/Retime/CurveEditorRetimeTool.cpp b/Engine/Plugins/Editor/CurveEditorTools/Source/CurveEditorTools/Private/Tools/Retime/CurveEditorRetimeTool.cpp
--- a/Engine/Plugins/Editor/CurveEditorTools/Source/CurveEditorTools/Private/Tools/Retime/CurveEditorRetimeTool.cpp
+++ b/Engine/Plugins/Editor/CurveEditorTools/Source/CurveEditorTools/Private/Tools/Retime/CurveEditorRetimeTool.cpp
@@ -1,6 +1,7 @@
 // Copyright Epic Games, Inc. All Rights Reserved.
 
 #include "CurveEditorRetimeTool.h"
+#include "Async/ParallelFor.h"
 #include "CurveEditorToolCommands.h"
 #include "CurveEditor.h"
 #include "CurveEditorCommands.h"
@@ -527,102 +528,125 @@ void FCurveEditorRetimeTool::OnDrag(const FVector2D& InStartPosition, const FVec
 void FCurveEditorRetimeTool::OnDrag_RetimeKeysAndAnchors(const double DragDeltaFromStartPx)
 {
 	OnDrag_MoveAnchors(DragDeltaFromStartPx);
-
 	const TSharedPtr<FCurveEditor> CurveEditor = WeakCurveEditor.Pin();
+
 	if (CurveEditor)
 	{
-		// For each curve we're affecting.
-		for (Retime::FPreDragChannelData& ChannelData : MouseDownOp->PreDragCurveData->CurveChannels)
+		// Profiled 2nd October, 2025, 5.7. See CL 46467607 description.
+		ParallelFor(MouseDownOp->PreDragCurveData->CurveChannels.Num(),[this, CurveEditor, DragDeltaFromStartPx](const int32 ChannelIndex)
 		{
-			const SCurveEditorView* View = CurveEditor->FindFirstInteractiveView(ChannelData.CurveID);
-			if (!View)
-			{
-				continue;
-			}
+			RecomputeLastDraggedFrameNumbersFromDrag(*CurveEditor, ChannelIndex, DragDeltaFromStartPx); 
+		});
 
+		// Need to loop through it again because SetKeyPositions is only safe to call on the game thread.
+		for (Retime::FPreDragChannelData& ChannelData : MouseDownOp->PreDragCurveData->CurveChannels)
+		{
 			FCurveModel* Curve = CurveEditor->FindCurve(ChannelData.CurveID);
-			if (!ensureAlways(Curve))
+			if (!Curve)
 			{
 				continue;
 			}
 
-			const FCurveEditorScreenSpace CurveSpace = View->GetCurveSpace(ChannelData.CurveID);
-			const double DeltaInputValue = DragDeltaFromStartPx / CurveSpace.PixelsPerInput();
+			// Profiled 2nd October, 2025, 5.7. See CL 46467607 description.
+			const TArrayView<const FKeyHandle> Keys = ChannelData.Handles;
+			const TArrayView<const FKeyPosition> Positions = ChannelData.LastDraggedFrameNumbers;
+			Curve->SetKeyPositions(Keys, Positions, EPropertyChangeType::Interactive);	
+		}
+	}
+	
+	// Sort to ensure anchors are always sorted from least to greatest as drawing and between segments counts on this behavior.
+	SortAnchorsByTime();
+}
 
-			// For each key on the channel we need to figure out which two anchors it's affected by so we can modify the position based on whether or not those anchors moved.
-			for (int32 KeyIndex = 0; KeyIndex < ChannelData.Handles.Num(); KeyIndex++)
-			{
-				const FKeyPosition& Position = ChannelData.FrameNumbers[KeyIndex];
+void FCurveEditorRetimeTool::RecomputeLastDraggedFrameNumbersFromDrag(
+	const FCurveEditor& CurveEditor, const int32 ChannelIndex, const double DragDeltaFromStartPx
+	)
+{
+	Retime::FPreDragChannelData& ChannelData = MouseDownOp->PreDragCurveData->CurveChannels[ChannelIndex];
+	const SCurveEditorView* View = CurveEditor.FindFirstInteractiveView(ChannelData.CurveID);
+	if (!View)
+	{
+		return;
+	}
 
-				// We need to find the two nearest handles for the key. A key may fall outside of any handle ranges so there may not be one.
-				// In this case we can treat the two handles as the same and give it a movement weight of 1.0. Otherwise we calculate the weight
-				// based on the fraction of the distance between the two handles (unless both are selected at which point it's full movement.
-				int32 LeftAnchorIndex = -1, RightAnchorIndex = -1;
+	FCurveModel* Curve = CurveEditor.FindCurve(ChannelData.CurveID);
+	if (!ensureAlways(Curve))
+	{
+		return;
+	}
 
-				double OriginalKeyTime = Position.InputValue;
-				for (int32 AnchorIndex = 0; AnchorIndex < MouseDownOp->PreDragCurveData->RetimeAnchorStartTimes.Num(); AnchorIndex++)
-				{
-					double OriginalAnchorTime = MouseDownOp->PreDragCurveData->RetimeAnchorStartTimes[AnchorIndex];
-
-					// We know anchors are sorted from lowest value to highest value, so we need to find the highest anchor that doesn't go past our key time to find the left one.
-					if (OriginalAnchorTime < OriginalKeyTime)
-					{
-						LeftAnchorIndex = AnchorIndex;
-					}
-					else if (OriginalAnchorTime > OriginalKeyTime)
-					{
-						// This means that this anchor is actually already to the right of the key, so we don't want to update our Left Key.
-						break;
-					}
-				}
+	const FCurveEditorScreenSpace CurveSpace = View->GetCurveSpace(ChannelData.CurveID);
+	const double DeltaInputValue = DragDeltaFromStartPx / CurveSpace.PixelsPerInput();
+	
+	// For each key on the channel we need to figure out which two anchors it's affected by so we can modify the position based on whether or not those anchors moved
+	for (int32 KeyIndex = 0; KeyIndex < ChannelData.Handles.Num(); KeyIndex++)
+	{
+		const FKeyPosition& Position = ChannelData.FrameNumbers[KeyIndex];
 
-				// We can assume the right anchor will be the left anchor's index (since the array is sorted)... assuming there's enough keys. If there's not enough keys then both anchors end up with the same index.
-				if (LeftAnchorIndex + 1 < RetimeData->RetimingAnchors.Num())
-				{
-					RightAnchorIndex = LeftAnchorIndex + 1;
-				}
-				else
-				{
-					RightAnchorIndex = LeftAnchorIndex;
-				}
+		// We need to find the two nearest handles for the key. A key may fall outside of any handle ranges so there may not be one.
+		// In this case we can treat the two handles as the same and give it a movement weight of 1.0. Otherwise we calculate the weight
+		// based on the fraction of the distance between the two handles (unless both are selected at which point it's full movement.
+		int32 LeftAnchorIndex = INDEX_NONE , RightAnchorIndex = INDEX_NONE ;
 
-				// There's no guarantee that the anchor indexes are valid yet! If the Left Anchor index was never found (too far to the right of the key we're looking at) it'll be -1 still, but Right Anchor will be 0.
-				// We know there's at least one retiming anchor (to be in this block of code) so we can assume if the left anchor wasn't found we can use index 0 for it as well.
-				if (LeftAnchorIndex == -1)
-				{
-					LeftAnchorIndex = 0;
-					RightAnchorIndex = 0;
-				}
+		const double OriginalKeyTime = Position.InputValue;
+		
+		for (int32 AnchorIndex = 0; AnchorIndex < MouseDownOp->PreDragCurveData->RetimeAnchorStartTimes.Num(); AnchorIndex++)
+		{
+			double OriginalAnchorTime = MouseDownOp->PreDragCurveData->RetimeAnchorStartTimes[AnchorIndex];
 
-				// Okay, now we know both Left and Right indexes are valid. They may point to the same anchor, but that's okay. Now we need to calculate the amount of influence on this key.
-				// If they do point to the same key, the math below doesn't actually end up mattering because the lerp goes between the same value, which is either full on or full off based on whether or not
-				// the nearest key was selected.
-				const double OriginalLeftAnchorTime = MouseDownOp->PreDragCurveData->RetimeAnchorStartTimes[LeftAnchorIndex];
-				const double OriginalRightAnchorTime = MouseDownOp->PreDragCurveData->RetimeAnchorStartTimes[RightAnchorIndex];
-				const double LeftAnchorInfluence = RetimeData->RetimingAnchors[LeftAnchorIndex].bIsSelected ? 1.0 : 0;
-				const double RightAnchorInfluence = RetimeData->RetimingAnchors[RightAnchorIndex].bIsSelected ? 1.0 : 0;
+			// We know anchors are sorted from lowest value to highest value, so we need to find the highest anchor that doesn't go past our key time to find the left one.
+			if (OriginalAnchorTime < OriginalKeyTime)
+			{
+				LeftAnchorIndex = AnchorIndex;
+			}
+			else if (OriginalAnchorTime > OriginalKeyTime)
+			{
+				// This means that this anchor is actually already to the right of the key, so we don't want to update our Left Key.
+				break;
+			}
+		}
 
-				// Figure out the fraction [0-1] between the two anchors our key is.
-				const double Fraction = FMath::GetMappedRangeValueClamped(FVector2D(OriginalLeftAnchorTime, OriginalRightAnchorTime), FVector2D(0, 1), OriginalKeyTime);
+		// We can assume the right anchor will be the left anchor's index (since the array is sorted)... assuming there's enough keys. If there's not enough keys then both anchors end up with the same index.
+		if (LeftAnchorIndex + 1 < RetimeData->RetimingAnchors.Num())
+		{
+			RightAnchorIndex = LeftAnchorIndex + 1;
+		}
+		else
+		{
+			RightAnchorIndex = LeftAnchorIndex;
+		}
+
+		// There's no guarantee that the anchor indexes are valid yet! If the Left Anchor index was never found (too far to the right of the key we're looking at) it'll be -1 still, but Right Anchor will be 0.
+		// We know there's at least one retiming anchor (to be in this block of code) so we can assume if the left anchor wasn't found we can use index 0 for it as well.
+		if (LeftAnchorIndex == INDEX_NONE )
+		{
+			LeftAnchorIndex = 0;
+			RightAnchorIndex = 0;
+		}
 
-				// Now we can remap between the two influences based on fraction amount.
-				const double InfluenceOnKey = FMath::Lerp(LeftAnchorInfluence, RightAnchorInfluence, Fraction);
+		// Okay, now we know both Left and Right indexes are valid. They may point to the same anchor, but that's okay. Now we need to calculate the amount of influence on this key.
+		// If they do point to the same key, the math below doesn't actually end up mattering because the lerp goes between the same value, which is either full on or full off based on whether or not
+		// the nearest key was selected.
+		const double OriginalLeftAnchorTime = MouseDownOp->PreDragCurveData->RetimeAnchorStartTimes[LeftAnchorIndex];
+		const double OriginalRightAnchorTime = MouseDownOp->PreDragCurveData->RetimeAnchorStartTimes[RightAnchorIndex];
+		const double LeftAnchorInfluence = RetimeData->RetimingAnchors[LeftAnchorIndex].bIsSelected ? 1.0 : 0;
+		const double RightAnchorInfluence = RetimeData->RetimingAnchors[RightAnchorIndex].bIsSelected ? 1.0 : 0;
 
-				double NewKeyValue = OriginalKeyTime + (DeltaInputValue * InfluenceOnKey);
+		// Figure out the fraction [0-1] between the two anchors our key is.
+		const double Fraction = FMath::GetMappedRangeValueClamped(FVector2D(OriginalLeftAnchorTime, OriginalRightAnchorTime), FVector2D(0, 1), OriginalKeyTime);
 
-				// Snap the new values to the grid.
-				FKeyPosition NewPosition = Position;
-				NewPosition.InputValue = View->IsTimeSnapEnabled() ? CurveEditor->GetCurveSnapMetrics(ChannelData.CurveID).SnapInputSeconds(NewKeyValue) : NewKeyValue;
+		// Now we can remap between the two influences based on fraction amount.
+		const double InfluenceOnKey = FMath::Lerp(LeftAnchorInfluence, RightAnchorInfluence, Fraction);
 
-				// Update our Curve Model with our updated data.
-				Curve->SetKeyPositions(MakeArrayView(&ChannelData.Handles[KeyIndex], 1), MakeArrayView(&NewPosition, 1), EPropertyChangeType::Interactive);
-				ChannelData.LastDraggedFrameNumbers[KeyIndex] = NewPosition;
-			}
-		}
+		const double NewKeyValue = OriginalKeyTime + (DeltaInputValue * InfluenceOnKey);
+		
+		// Snap the new values to the grid.
+		FKeyPosition NewPosition = Position;
+		NewPosition.InputValue = View->IsTimeSnapEnabled() ? CurveEditor.GetCurveSnapMetrics(ChannelData.CurveID).SnapInputSeconds(NewKeyValue) : NewKeyValue;
+
+		// Update our Curve Model with our updated data.				
+		ChannelData.LastDraggedFrameNumbers[KeyIndex] = NewPosition;				
 	}
-	
-	// Sort to ensure anchors are always sorted from least to greatest as drawing and between segments counts on this behavior.
-	SortAnchorsByTime();
 }
 
 void FCurveEditorRetimeTool::OnDrag_MoveAnchors(const double DragDeltaFromStartPx)
@@ -652,24 +676,26 @@ void FCurveEditorRetimeTool::OnDrag_MoveAnchors(const double DragDeltaFromStartP
 void FCurveEditorRetimeTool::OnDragEnd()
 {
 	using namespace Retime;
-	
-	TSharedPtr<FCurveEditor> CurveEditor = WeakCurveEditor.Pin();
 
+	const TSharedPtr<FCurveEditor> CurveEditor = WeakCurveEditor.Pin();
+	if (!CurveEditor)
+	{
+		return;
+	}
 	// For each curve we're affecting.
 	for (FPreDragChannelData& ChannelData : MouseDownOp->PreDragCurveData->CurveChannels)
 	{
-		for (int32 KeyIndex = 0; KeyIndex < ChannelData.Handles.Num(); KeyIndex++)
+		FCurveModel* Curve = CurveEditor->FindCurve(ChannelData.CurveID);
+		if (!ensureAlways(Curve))
 		{
-			if (CurveEditor)
-			{
-				FCurveModel* Curve = CurveEditor->FindCurve(ChannelData.CurveID);
-				if (Curve)
-				{
-					Curve->SetKeyPositions(MakeArrayView(&ChannelData.Handles[KeyIndex], 1), MakeArrayView(&ChannelData.LastDraggedFrameNumbers[KeyIndex], 1), EPropertyChangeType::ValueSet);
-				}
-			}
+			continue;
 		}
+		
+		const TArrayView<const FKeyHandle> Keys = ChannelData.Handles;
+		const TArrayView<const FKeyPosition> Positions = ChannelData.LastDraggedFrameNumbers;
+		Curve->SetKeyPositions(Keys, Positions, EPropertyChangeType::ValueSet);	
 	}
+	
 	if (MouseDownOp->CurrentDragMode == EDragMode::MoveAnchorMode)
 	{
 		// When adjusting anchor position only sort anchors at the end of a drag to avoid unnecessary operations during anchor movement.
diff --git a/Engine/Plugins/Editor/CurveEditorTools/Source/CurveEditorTools/Private/Tools/Retime/CurveEditorRetimeTool.h b/Engine/Plugins/Editor/CurveEditorTools/Source/CurveEditorTools/Private/Tools/Retime/CurveEditorRetimeTool.h
--- a/Engine/Plugins/Editor/CurveEditorTools/Source/CurveEditorTools/Private/Tools/Retime/CurveEditorRetimeTool.h
+++ b/Engine/Plugins/Editor/CurveEditorTools/Source/CurveEditorTools/Private/Tools/Retime/CurveEditorRetimeTool.h
@@ -67,7 +67,9 @@ private:
 	void OnDrag_MoveAnchors(const double DragDeltaFromStartPx);
 	/** Sort all anchors by input time to ensure they're always sorted from least to greatest as drawing and between segments counts on this behaviour. */
 	void SortAnchorsByTime();
-	
+	/** Calculate positions of keys on the curve after retiming and save them to ChannelData.LastDraggedFrameNumbers for setting the key positions in OnDrag_RetimeKeysAndAnchors. */
+	void RecomputeLastDraggedFrameNumbersFromDrag(const FCurveEditor& CurveEditor, const int32 ChannelIndex, const double DragDeltaFromStartPx);
+
 	/** Returns true if the key event is a scrub time key event */
 	bool IsScrubTimeKeyEvent(const FKeyEvent& InKeyEvent) const;
 
