Curve Editor: Fix lattice tool no longer working correctly.

The problem was OnCurveModified being invoked while a drag operation was technically taking place. The fix is to set the bIsModifying flag when the operation starts and reset it when it ends. Before it would temporarily set it once before each SetKeyPositions is done.

diff --git a/Engine/Plugins/Editor/CurveEditorTools/Source/CurveEditorTools/Private/Tools/Lattice/CurveEditorLatticeTool.cpp b/Engine/Plugins/Editor/CurveEditorTools/Source/CurveEditorTools/Private/Tools/Lattice/CurveEditorLatticeTool.cpp
--- a/Engine/Plugins/Editor/CurveEditorTools/Source/CurveEditorTools/Private/Tools/Lattice/CurveEditorLatticeTool.cpp
+++ b/Engine/Plugins/Editor/CurveEditorTools/Source/CurveEditorTools/Private/Tools/Lattice/CurveEditorLatticeTool.cpp
@@ -184,9 +184,6 @@ FReply FCurveEditorLatticeTool::OnMouseButtonUp(
 {
 	if (DelayedDragData.IsSet())
 	{
-		// Resetting DelayedDragData may cause FScopedCurveChange to invoke the curves' modified events... which would cause us to recompute our lattice grid.
-		const TGuardValue Guard(LatticeDeformerState->bIsModifyingCurves, true);
-		
 		FReply Reply = FReply::Handled();
 		if (DelayedDragData->DelayedDrag.IsDragging() && ensure(DelayedDragData->ActiveOperation))
 		{
@@ -195,6 +192,9 @@ FReply FCurveEditorLatticeTool::OnMouseButtonUp(
 			Reply.ReleaseMouseCapture();
 		}
 		
+		// Resetting DelayedDragData may cause FScopedCurveChange to invoke the curves' modified events... which would cause us to recompute our lattice grid.
+		// It's important to do this after EndDrag is (potentially) called, as that sets LatticeDeformerState->bIsModifyingCurves to false.
+		const TGuardValue Guard(LatticeDeformerState->bIsModifyingCurves, true);
 		DelayedDragData.Reset();
 		return Reply;
 	}
@@ -534,15 +534,17 @@ TOptional<FCurveEditorLatticeTool::FDragOpData> FCurveEditorLatticeTool::CreateD
 				const bool bIsDragFinished = ChangeType == EPropertyChangeType::ValueSet; 
 				if (bIsDragFinished)
 				{
+					LatticeDeformerState->bIsModifyingCurves = false;
 					SaveUndoState(*WeakCurveEditor.Pin());
 				}
 			});
 	};
+	const FSimpleDelegate HandleBeginDrag = FSimpleDelegate::CreateLambda([this](){ LatticeDeformerState->bIsModifyingCurves = true; });
 	const FHandleMovedPoints HandleDrag = MakeDelegate(EPropertyChangeType::Interactive);
 	const FHandleMovedPoints HandleFinishDrag = MakeDelegate(EPropertyChangeType::ValueSet);
 	
 	TUniquePtr<FLatticeDragOp> Op = MakeUnique<FLatticeDragOp_MoveControlPoints>(
-		WeakCurveEditor, FControlPointArray{ ControlPoints[InControlPointIndex] }, InBounds, HandleDrag, HandleFinishDrag
+		WeakCurveEditor, FControlPointArray{ ControlPoints[InControlPointIndex] }, InBounds, HandleBeginDrag, HandleDrag, HandleFinishDrag
 	);
 	FDragOpData DragDrop(
 		WeakCurveEditor, MoveTemp(Op),
@@ -584,16 +586,18 @@ TOptional<FCurveEditorLatticeTool::FDragOpData> FCurveEditorLatticeTool::CreateD
 				const bool bIsDragFinished = ChangeType == EPropertyChangeType::ValueSet; 
 				if (bIsDragFinished)
 				{
+					LatticeDeformerState->bIsModifyingCurves = false;
 					SaveUndoState(*WeakCurveEditor.Pin());
 				}
 			});
 	};
+	const FSimpleDelegate HandleBeginDrag = FSimpleDelegate::CreateLambda([this](){ LatticeDeformerState->bIsModifyingCurves = true; });
 	const FHandleMovedPoints HandleDrag = MakeDelegate(EPropertyChangeType::Interactive);
 	const FHandleMovedPoints HandleFinishDrag = MakeDelegate(EPropertyChangeType::ValueSet);
 		
 	FControlPointArray EdgePoints{ ControlPoints[EdgeIndices.Start], ControlPoints[EdgeIndices.End] };
 	TUniquePtr<FLatticeDragOp> Op = MakeUnique<FLatticeDragOp_MoveControlPoints>(
-		WeakCurveEditor, MoveTemp(EdgePoints), InBounds, HandleDrag, HandleFinishDrag
+		WeakCurveEditor, MoveTemp(EdgePoints), InBounds, HandleBeginDrag, HandleDrag, HandleFinishDrag
 	);
 	FDragOpData DragDrop(
 		WeakCurveEditor, MoveTemp(Op),
@@ -636,17 +640,19 @@ TOptional<FCurveEditorLatticeTool::FDragOpData> FCurveEditorLatticeTool::CreateD
 				const bool bIsDragFinished = ChangeType == EPropertyChangeType::ValueSet; 
 				if (bIsDragFinished)
 				{
+					LatticeDeformerState->bIsModifyingCurves = false;
 					SaveUndoState(*WeakCurveEditor.Pin());
 				}
 			});
 	};
+	const FSimpleDelegate HandleBeginDrag = FSimpleDelegate::CreateLambda([this](){ LatticeDeformerState->bIsModifyingCurves = true; });
 	const FHandleMovedPoints HandleDrag = MakeDelegate(EPropertyChangeType::Interactive);
 	const FHandleMovedPoints HandleFinishDrag = MakeDelegate(EPropertyChangeType::ValueSet);
 		
 	const auto[TopLeft, TopRight, BottomRight, BottomLeft] = CellIndices;
 	FControlPointArray CellPoints{ ControlPoints[TopLeft], ControlPoints[TopRight], ControlPoints[BottomRight], ControlPoints[BottomLeft] };
 	TUniquePtr<FLatticeDragOp> Op = MakeUnique<FLatticeDragOp_MoveControlPoints>(
-		WeakCurveEditor, MoveTemp(CellPoints), InBounds, HandleDrag, HandleFinishDrag
+		WeakCurveEditor, MoveTemp(CellPoints), InBounds, HandleBeginDrag, HandleDrag, HandleFinishDrag
 	);
 	return FDragOpData(
 		WeakCurveEditor, MoveTemp(Op),
@@ -742,7 +748,7 @@ void FCurveEditorLatticeTool::OnMoveControlPoint(
 	)
 {
 	check(InVertexIndices.Num() == 1 && InNewControlPointPositions.Num() == 1);
-	TGuardValue Guard(LatticeDeformerState->bIsModifyingCurves, true);
+	ensure(LatticeDeformerState->bIsModifyingCurves);
 	
 	const FCurveModelID& CurveId = *InCurveModel.GetId();
 	if (FLatticePointTangentsMirrorOp* MirrorOp = DelayedDragData->ActiveOperation->PointTangentMirroringOps.Find(CurveId))
@@ -760,7 +766,7 @@ void FCurveEditorLatticeTool::OnMoveEdge(
 	)
 {
 	check(InVertexIndices.Num() == 2 && InNewControlPointPositions.Num() == 2);
-	TGuardValue Guard(LatticeDeformerState->bIsModifyingCurves, true);
+	ensure(LatticeDeformerState->bIsModifyingCurves);
 
 	const FCurveModelID& CurveId = *InCurveModel.GetId();
 	if (FLatticeEdgeTangentsMirrorOp* MirrorOp = DelayedDragData->ActiveOperation->EdgeTangentMirroringOps.Find(CurveId))
@@ -777,6 +783,8 @@ void FCurveEditorLatticeTool::OnSetControlPoints(
 	EPropertyChangeType::Type InChangeType
 	)
 {
+	ensure(LatticeDeformerState->bIsModifyingCurves);
+	
 	const TSharedPtr<FCurveEditor> CurveEditor = WeakCurveEditor.Pin();
 	const TOptional<FCurvePointSnapper> Snapper = CurveEditor ? FCurvePointSnapper::MakeSnapper(*CurveEditor) : TOptional<FCurvePointSnapper>{};
 	if (!CurveEditor || !Snapper)
@@ -784,8 +792,6 @@ void FCurveEditorLatticeTool::OnSetControlPoints(
 		return;
 	}
 	
-	TGuardValue Guard(LatticeDeformerState->bIsModifyingCurves, true);
-	
 	// Combining Handles & Keys into large arrays like this is significantly faster than feeding each key into SetKeyPositions separately.
 	// For 9000 keys it goes from 2136 ms to 238 ms.
 	TArray<FKeyHandle> Handles;
@@ -822,6 +828,8 @@ void FCurveEditorLatticeTool::MoveEdgeToOpposite(FPerCurveDeformer2D& InDeformer
 		return;
 	}
 	
+	const TGuardValue Guard(LatticeDeformerState->bIsModifyingCurves, true);
+	
 	// Flatten tangents if moving top or bottom edge...
 	const FCurveEditor& CurveEditor = *WeakCurveEditor.Pin();
 	TOptional<FLatticeEdgeTangentsMirrorOp> MirrorOp = FLatticeEdgeTangentsMirrorOp::MakeMirrorOpForDragLatticeEdge(
diff --git a/Engine/Plugins/Editor/CurveEditorTools/Source/CurveEditorTools/Private/Tools/Lattice/Drag/LatticeDragOp_MoveControlPoints.cpp b/Engine/Plugins/Editor/CurveEditorTools/Source/CurveEditorTools/Private/Tools/Lattice/Drag/LatticeDragOp_MoveControlPoints.cpp
--- a/Engine/Plugins/Editor/CurveEditorTools/Source/CurveEditorTools/Private/Tools/Lattice/Drag/LatticeDragOp_MoveControlPoints.cpp
+++ b/Engine/Plugins/Editor/CurveEditorTools/Source/CurveEditorTools/Private/Tools/Lattice/Drag/LatticeDragOp_MoveControlPoints.cpp
@@ -11,12 +11,14 @@ FLatticeDragOp_MoveControlPoints::FLatticeDragOp_MoveControlPoints(
 	TWeakPtr<FCurveEditor> InCurveEditor,
 	FControlPointArray InInitialControlPoints,
 	const FLatticeBounds& InBounds,
+	FOnBeginDragOperation InOnBeginDragControlPointsDelegate,
 	FOnControlPointsMoved InInteractiveDragControlPointsDelegate,
 	FOnControlPointsMoved InFinishDragControlPointsDelegate
 	)
 	: FLatticeDragOp(MoveTemp(InCurveEditor))
 	, InitialControlPoints(MoveTemp(InInitialControlPoints))
 	, Bounds(InBounds)
+	, OnBeginDragControlPointsDelegate(MoveTemp(InOnBeginDragControlPointsDelegate))
 	, OnInteractiveDragControlPointsDelegate(MoveTemp(InInteractiveDragControlPointsDelegate))
 	, OnFinishDragControlPointsDelegate(MoveTemp(InFinishDragControlPointsDelegate))
 	, UpdatedControlPointPositions(InitialControlPoints)
@@ -24,7 +26,12 @@ FLatticeDragOp_MoveControlPoints::FLatticeDragOp_MoveControlPoints(
 	UE_LOG(LogTemp, Warning, TEXT("SlateMin: %s \t SlateMax:%s"), *Bounds.MinSlatePosition.ToString(), *Bounds.MaxSlatePosition.ToString());
 	UE_LOG(LogTemp, Warning, TEXT("Min: %s \t Max:%s"), *Bounds.MinValues.ToString(), *Bounds.MaxValues.ToString());
 }
-	
+
+void FLatticeDragOp_MoveControlPoints::OnBeginDrag(const FGeometry& InGeometry, const FVector2D& InInitialMousePosition)
+{
+	OnBeginDragControlPointsDelegate.ExecuteIfBound();
+}
+
 void FLatticeDragOp_MoveControlPoints::OnMoveMouse(const FGeometry& InGeometry, const FVector2D& InScreenPosition)
 {
 	UpdateControlPoints(InGeometry, InScreenPosition);
diff --git a/Engine/Plugins/Editor/CurveEditorTools/Source/CurveEditorTools/Private/Tools/Lattice/Drag/LatticeDragOp_MoveControlPoints.h b/Engine/Plugins/Editor/CurveEditorTools/Source/CurveEditorTools/Private/Tools/Lattice/Drag/LatticeDragOp_MoveControlPoints.h
--- a/Engine/Plugins/Editor/CurveEditorTools/Source/CurveEditorTools/Private/Tools/Lattice/Drag/LatticeDragOp_MoveControlPoints.h
+++ b/Engine/Plugins/Editor/CurveEditorTools/Source/CurveEditorTools/Private/Tools/Lattice/Drag/LatticeDragOp_MoveControlPoints.h
@@ -21,12 +21,14 @@ class FLatticeDragOp_MoveControlPoints : public FLatticeDragOp
 {
 public:
 
+	using FOnBeginDragOperation = FSimpleDelegate;
 	DECLARE_DELEGATE_OneParam(FOnControlPointsMoved, const TConstArrayView<FVector2D>& InNewControlPointPositions);
 
 	explicit FLatticeDragOp_MoveControlPoints(
 		TWeakPtr<FCurveEditor> InCurveEditor,
 		FControlPointArray InInitialControlPoints,
 		const FLatticeBounds& InBounds,
+		FOnBeginDragOperation InOnBeginDragControlPointsDelegate,
 		FOnControlPointsMoved InInteractiveDragControlPointsDelegate,
 		FOnControlPointsMoved InFinishDragControlPointsDelegate
 		);
@@ -34,6 +36,7 @@ public:
 protected:
 	
 	//~ Begin FLatticeDragOp Interface
+	virtual void OnBeginDrag(const FGeometry& InGeometry, const FVector2D& InInitialMousePosition) override;
 	virtual void OnMoveMouse(const FGeometry& InGeometry, const FVector2D& InScreenPosition) override;
 	virtual void OnEndDrag(const FGeometry& InGeometry, const FVector2D& InMousePosition) override;
 	//~ End FLatticeDragOp Interface
@@ -44,7 +47,9 @@ private:
 	const FControlPointArray InitialControlPoints;
 	/** Used to transform points */
 	const FLatticeBounds Bounds;
-	
+
+	/** Invoked when the drag operation begins. A call to OnFinishDragControlPointsDelegate will follow some time in the future. */
+	const FOnBeginDragOperation OnBeginDragControlPointsDelegate;
 	/** Invoked while the mouse button is down. */
 	const FOnControlPointsMoved OnInteractiveDragControlPointsDelegate;
 	/** Invoked one last time when the mouse button is lifted. */
