Curve Editor: Fix snapping when only moving horizontally or vertically

diff --git a/Engine/Plugins/Editor/CurveEditorTools/Source/CurveEditorTools/Private/Misc/CurvePointSnapper.cpp b/Engine/Plugins/Editor/CurveEditorTools/Source/CurveEditorTools/Private/Misc/CurvePointSnapper.cpp
--- a/Engine/Plugins/Editor/CurveEditorTools/Source/CurveEditorTools/Private/Misc/CurvePointSnapper.cpp
+++ b/Engine/Plugins/Editor/CurveEditorTools/Source/CurveEditorTools/Private/Misc/CurvePointSnapper.cpp
@@ -9,36 +9,37 @@
 
 namespace UE::CurveEditorTools
 {
-TOptional<FCurvePointSnapper> FCurvePointSnapper::MakeSnapper(FCurveEditor& InCurveEditor)
+TOptional<FCurvePointSnapper> FCurvePointSnapper::MakeSnapper(FCurveEditor& InCurveEditor, const FCurveModelID& InCurveModelId)
 {
 	if (InCurveEditor.GetSelection().GetAll().IsEmpty())
 	{
 		return {};
 	}
 	
-	// Snap the mouse to the grid according to the first curve's snap metrics (assuming all curves have the same view scales)
-	const FCurveModelID FirstCurveID = InCurveEditor.GetSelection().GetAll().CreateConstIterator()->Key;
-	const SCurveEditorView* View = InCurveEditor.FindFirstInteractiveView(FirstCurveID);
+	const SCurveEditorView* View = InCurveEditor.FindFirstInteractiveView(InCurveModelId);
 	if (!View)
 	{
 		return {};
 	}
 
-	return FCurvePointSnapper(View->GetCurveSpace(FirstCurveID), InCurveEditor.GetCurveSnapMetrics(FirstCurveID));
+	return FCurvePointSnapper(View->GetCurveSpace(InCurveModelId), InCurveEditor.GetCurveSnapMetrics(InCurveModelId));
 }
 
-FVector2D FCurvePointSnapper::SnapPoint(const FVector2D& InCurveSpacePoint) const
+FVector2D FCurvePointSnapper::SnapPoint(const FVector2D& InCurveSpacePoint, EAxisSnapMode InDragMode) const
 {
+	// LockedX & LockedY mean that the user is e.g. shift dragging: this means they're intentionally moving in only one direction.
 	return FVector2D
 	{
-		SnapMetrics.SnapInputSeconds(InCurveSpacePoint.X),
-		SnapMetrics.SnapOutput(InCurveSpacePoint.Y)
+		// If moving only up & down, then it would be weird if the keys snapped in x direction.
+		EnumHasAnyFlags(InDragMode, EAxisSnapMode::IgnoreSnapOnX) ? InCurveSpacePoint.X : SnapMetrics.SnapInputSeconds(InCurveSpacePoint.X),
+		// If moving only left & right, then it would be weird if the keys snapped in y direction.
+		EnumHasAnyFlags(InDragMode, EAxisSnapMode::IgnoreSnapOnY) ? InCurveSpacePoint.Y : SnapMetrics.SnapOutput(InCurveSpacePoint.Y)
 	};
 }
 
-FKeyPosition FCurvePointSnapper::SnapKey(const FKeyPosition& InKeyPosition) const
+FKeyPosition FCurvePointSnapper::SnapKey(const FKeyPosition& InKeyPosition, EAxisSnapMode InDragMode) const
 {
-	const FVector2D Point = SnapPoint({ InKeyPosition.InputValue, InKeyPosition.OutputValue });
+	const FVector2D Point = SnapPoint({ InKeyPosition.InputValue, InKeyPosition.OutputValue }, InDragMode);
 	return { Point.X, Point.Y };
 }
 }
diff --git a/Engine/Plugins/Editor/CurveEditorTools/Source/CurveEditorTools/Private/Misc/CurvePointSnapper.h b/Engine/Plugins/Editor/CurveEditorTools/Source/CurveEditorTools/Private/Misc/CurvePointSnapper.h
--- a/Engine/Plugins/Editor/CurveEditorTools/Source/CurveEditorTools/Private/Misc/CurvePointSnapper.h
+++ b/Engine/Plugins/Editor/CurveEditorTools/Source/CurveEditorTools/Private/Misc/CurvePointSnapper.h
@@ -4,7 +4,9 @@
 
 #include "CurveEditorScreenSpace.h"
 #include "CurveEditorSnapMetrics.h"
+#include "Tools/Lattice/Drag/EAxisSnapMode.h"
 
+struct FCurveModelID;
 struct FKeyPosition;
 class FCurveEditor;
 template<typename OptionalType> struct TOptional;
@@ -17,13 +19,13 @@ class FCurvePointSnapper
 public:
 
 	/** @return Unset if there's no curve view from which to get the FCurveEditorScreenSpace transform. */
-	static TOptional<FCurvePointSnapper> MakeSnapper(FCurveEditor& InCurveEditor);
+	static TOptional<FCurvePointSnapper> MakeSnapper(FCurveEditor& InCurveEditor, const FCurveModelID& InCurveModelId);
 
 	/** Snaps the point as per the snapping settings */
-	FVector2D SnapPoint(const FVector2D& InCurveSpacePoint) const;
+	FVector2D SnapPoint(const FVector2D& InCurveSpacePoint, EAxisSnapMode InDragMode = EAxisSnapMode::Unrestricted) const;
 
 	/** Snaps the key as per the snapping settings */
-	FKeyPosition SnapKey(const FKeyPosition& InKeyPosition) const;
+	FKeyPosition SnapKey(const FKeyPosition& InKeyPosition, EAxisSnapMode InDragMode = EAxisSnapMode::Unrestricted) const;
 
 private:
 
diff --git a/Engine/Plugins/Editor/CurveEditorTools/Source/CurveEditorTools/Private/Tools/Lattice/CurveEditorLatticeTool.cpp b/Engine/Plugins/Editor/CurveEditorTools/Source/CurveEditorTools/Private/Tools/Lattice/CurveEditorLatticeTool.cpp
--- a/Engine/Plugins/Editor/CurveEditorTools/Source/CurveEditorTools/Private/Tools/Lattice/CurveEditorLatticeTool.cpp
+++ b/Engine/Plugins/Editor/CurveEditorTools/Source/CurveEditorTools/Private/Tools/Lattice/CurveEditorLatticeTool.cpp
@@ -524,10 +524,10 @@ TOptional<FCurveEditorLatticeTool::FDragOpData> FCurveEditorLatticeTool::CreateD
 	const auto MakeDelegate = [this, InControlPointIndex](EPropertyChangeType::Type ChangeType)
 	{
 		return FHandleMovedPoints::CreateLambda(
-			[this, InControlPointIndex, ChangeType](const TConstArrayView<FVector2D>& InNewControlPointPositions)
+			[this, InControlPointIndex, ChangeType](const TConstArrayView<FVector2D>& InNewControlPointPositions, EAxisSnapMode InSnapMode)
 			{
 				const int32 Indices[1] { InControlPointIndex };
-				HandleDragPoints(&FCurveEditorLatticeTool::OnMoveControlPoint,Indices, InNewControlPointPositions, ChangeType);
+				HandleDragPoints(&FCurveEditorLatticeTool::OnMoveControlPoint,Indices, InNewControlPointPositions, ChangeType, InSnapMode);
 				GetGlobalDeformer().UpdateControlPoints(Indices, InNewControlPointPositions, [](auto, auto){});
 
 				// When the drag is finished, save where the lattice is so its shape can be transacted correctly.
@@ -576,10 +576,10 @@ TOptional<FCurveEditorLatticeTool::FDragOpData> FCurveEditorLatticeTool::CreateD
 	const auto MakeDelegate = [this, EdgeIndices](EPropertyChangeType::Type ChangeType)
 	{
 		return FHandleMovedPoints::CreateLambda(
-			[this, EdgeIndices, ChangeType](const TConstArrayView<FVector2D>& InNewControlPointPositions)
+			[this, EdgeIndices, ChangeType](const TConstArrayView<FVector2D>& InNewControlPointPositions, EAxisSnapMode InSnapMode)
 			{
 				const int32 Indices[2]{ EdgeIndices.Start, EdgeIndices.End };
-				HandleDragPoints(&FCurveEditorLatticeTool::OnMoveEdge, Indices, InNewControlPointPositions, ChangeType);
+				HandleDragPoints(&FCurveEditorLatticeTool::OnMoveEdge, Indices, InNewControlPointPositions, ChangeType, InSnapMode);
 				GetGlobalDeformer().UpdateControlPoints(Indices, InNewControlPointPositions, [](auto, auto){});
 
 				// When the drag is finished, save where the lattice is so its shape can be transacted correctly.
@@ -630,10 +630,10 @@ TOptional<FCurveEditorLatticeTool::FDragOpData> FCurveEditorLatticeTool::CreateD
 	const auto MakeDelegate = [this, CellIndices](EPropertyChangeType::Type ChangeType)
 	{
 		return FHandleMovedPoints::CreateLambda(
-			[this, CellIndices, ChangeType](const TConstArrayView<FVector2D>& InNewControlPointPositions)
+			[this, CellIndices, ChangeType](const TConstArrayView<FVector2D>& InNewControlPointPositions, EAxisSnapMode InSnapMode)
 			{
 				const int32 Indices[4]{ CellIndices.TopLeft, CellIndices.TopRight, CellIndices.BottomRight, CellIndices.BottomLeft };
-				HandleDragPoints(&FCurveEditorLatticeTool::OnSetControlPoints, Indices, InNewControlPointPositions, ChangeType);
+				HandleDragPoints(&FCurveEditorLatticeTool::OnSetControlPoints, Indices, InNewControlPointPositions, ChangeType, InSnapMode);
 				GetGlobalDeformer().UpdateControlPoints(Indices, InNewControlPointPositions, [](auto, auto){});
 				
 				// When the drag is finished, save where the lattice is so its shape can be transacted correctly.
@@ -661,8 +661,9 @@ TOptional<FCurveEditorLatticeTool::FDragOpData> FCurveEditorLatticeTool::CreateD
 }
 
 void FCurveEditorLatticeTool::HandleDragPoints(
-	FMovePointsSignature InMovePointsCallback, TConstArrayView<int32> InVertexIndices,
-	TConstArrayView<FVector2D> InNewControlPoints, EPropertyChangeType::Type InChangeType
+	FMovePointsSignature InMovePointsCallback,
+	TConstArrayView<int32> InVertexIndices, TConstArrayView<FVector2D> InNewControlPoints,
+	EPropertyChangeType::Type InChangeType, EAxisSnapMode InSnapMode
 	)
 {
 	const TSharedPtr<FCurveEditor> CurveEditorPin = WeakCurveEditor.Pin();
@@ -691,7 +692,7 @@ void FCurveEditorLatticeTool::HandleDragPoints(
 		const TArray<FVector2D> CurveRelativeControlPoints = TransformPoints(Pair.Value.GlobalDeformerToCurveDeformer, InNewControlPoints);
 		
 		FPerCurveDeformer2D& PerCurveDeformer = Pair.Value.CurveDeformer;
-		(this->*InMovePointsCallback)(PerCurveDeformer, *CurveModel, InVertexIndices, CurveRelativeControlPoints, InChangeType);
+		(this->*InMovePointsCallback)(PerCurveDeformer, *CurveModel, InVertexIndices, CurveRelativeControlPoints, InChangeType, InSnapMode);
 	}
 }
 
@@ -744,7 +745,7 @@ void FCurveEditorLatticeTool::HandleMirrorGrid()
 void FCurveEditorLatticeTool::OnMoveControlPoint(
 	FPerCurveDeformer2D& InDeformer, FCurveModel& InCurveModel,
 	TConstArrayView<int32> InVertexIndices, TConstArrayView<FVector2D> InNewControlPointPositions,
-	EPropertyChangeType::Type InChangeType
+	EPropertyChangeType::Type InChangeType, EAxisSnapMode InSnapMode
 	)
 {
 	check(InVertexIndices.Num() == 1 && InNewControlPointPositions.Num() == 1);
@@ -756,13 +757,13 @@ void FCurveEditorLatticeTool::OnMoveControlPoint(
 		MirrorOp->OnMovePoint(InNewControlPointPositions[0], *WeakCurveEditor.Pin());
 	}
 	
-	OnSetControlPoints(InDeformer, InCurveModel, InVertexIndices, InNewControlPointPositions, InChangeType);
+	OnSetControlPoints(InDeformer, InCurveModel, InVertexIndices, InNewControlPointPositions, InChangeType, InSnapMode);
 }
 
 void FCurveEditorLatticeTool::OnMoveEdge(
 	FPerCurveDeformer2D& InDeformer, FCurveModel& InCurveModel,
 	TConstArrayView<int32> InVertexIndices, TConstArrayView<FVector2D> InNewControlPointPositions,
-	EPropertyChangeType::Type InChangeType
+	EPropertyChangeType::Type InChangeType, EAxisSnapMode InSnapMode
 	)
 {
 	check(InVertexIndices.Num() == 2 && InNewControlPointPositions.Num() == 2);
@@ -774,19 +775,19 @@ void FCurveEditorLatticeTool::OnMoveEdge(
 		MirrorOp->OnMoveEdge(InNewControlPointPositions, *WeakCurveEditor.Pin());
 	}
 
-	OnSetControlPoints(InDeformer, InCurveModel, InVertexIndices, InNewControlPointPositions, InChangeType);
+	OnSetControlPoints(InDeformer, InCurveModel, InVertexIndices, InNewControlPointPositions, InChangeType, InSnapMode);
 }
 
 void FCurveEditorLatticeTool::OnSetControlPoints(
 	FPerCurveDeformer2D& InDeformer, FCurveModel& InCurveModel,
 	TConstArrayView<int32> InControlPointsToUpdate, TConstArrayView<FVector2D> InNewControlPoints,
-	EPropertyChangeType::Type InChangeType
+	EPropertyChangeType::Type InChangeType, EAxisSnapMode InSnapMode
 	)
 {
 	ensure(LatticeDeformerState->bIsModifyingCurves);
 	
 	const TSharedPtr<FCurveEditor> CurveEditor = WeakCurveEditor.Pin();
-	const TOptional<FCurvePointSnapper> Snapper = CurveEditor ? FCurvePointSnapper::MakeSnapper(*CurveEditor) : TOptional<FCurvePointSnapper>{};
+	const TOptional<FCurvePointSnapper> Snapper = CurveEditor ? FCurvePointSnapper::MakeSnapper(*CurveEditor, InCurveModel.GetOrInitId()) : TOptional<FCurvePointSnapper>{};
 	if (!CurveEditor || !Snapper)
 	{
 		return;
@@ -803,12 +804,10 @@ void FCurveEditorLatticeTool::OnSetControlPoints(
 	Keys.Reserve(NumKeys);
 	
 	InDeformer.UpdateControlPoints(InControlPointsToUpdate, InNewControlPoints,
-		[&Snapper, &Handles, &Keys](const FKeyHandle& KeyHandle, const FVector2D& KeyPosition)
+		[InSnapMode, &Snapper, &Handles, &Keys](const FKeyHandle& KeyHandle, const FVector2D& KeyPosition)
 		{
 			Handles.Add(KeyHandle);
-			Keys.Add(
-				Snapper->SnapKey(FKeyPosition{ KeyPosition.X, KeyPosition.Y })
-				);
+			Keys.Add(Snapper->SnapKey(FKeyPosition{ KeyPosition.X, KeyPosition.Y }, InSnapMode));
 		});
 
 	// Some stats all with 9000 keys. Varying degree of changed keys, i.e. how far you move your mouse in quick succession.
diff --git a/Engine/Plugins/Editor/CurveEditorTools/Source/CurveEditorTools/Private/Tools/Lattice/CurveEditorLatticeTool.h b/Engine/Plugins/Editor/CurveEditorTools/Source/CurveEditorTools/Private/Tools/Lattice/CurveEditorLatticeTool.h
--- a/Engine/Plugins/Editor/CurveEditorTools/Source/CurveEditorTools/Private/Tools/Lattice/CurveEditorLatticeTool.h
+++ b/Engine/Plugins/Editor/CurveEditorTools/Source/CurveEditorTools/Private/Tools/Lattice/CurveEditorLatticeTool.h
@@ -235,20 +235,37 @@ private:
 	TOptional<FDragOpData> CreateDragDrop_Cell(const FLatticeBounds& InBounds, int32 InCellIndex);
 
 	// User actions: things the user can directly do in the UI.
-	using FMovePointsSignature = void(FCurveEditorLatticeTool::*)(FPerCurveDeformer2D& InDeformer, FCurveModel& InCurveModel, TConstArrayView<int32> InVertexIndices, TConstArrayView<FVector2D> InNewControlPointPositions, EPropertyChangeType::Type InChangeType);
+	using FMovePointsSignature = void(FCurveEditorLatticeTool::*)(
+		FPerCurveDeformer2D& InDeformer, FCurveModel& InCurveModel,
+		TConstArrayView<int32> InVertexIndices, TConstArrayView<FVector2D> InNewControlPointPositions,
+		EPropertyChangeType::Type InChangeType, EAxisSnapMode InSnapMode
+		);
 	/** Handles dragging a control point, edge, or the lattice center. */
-	void HandleDragPoints(FMovePointsSignature InMovePointsCallback, TConstArrayView<int32> InVertexIndices, TConstArrayView<FVector2D> InNewControlPoints, EPropertyChangeType::Type InChangeType);
+	void HandleDragPoints(
+		FMovePointsSignature InMovePointsCallback, TConstArrayView<int32> InVertexIndices, TConstArrayView<FVector2D> InNewControlPoints,
+		EPropertyChangeType::Type InChangeType, EAxisSnapMode InSnapMode
+		);
 	/** Handles double-clicking an edge. */
 	void HandleMoveEdgeToOpposite(int32 InEdgeIndex);
 	/** Handles double-clicking the lattice center. */
 	void HandleMirrorGrid();
 	
 	/** Called when a control point is dragged. */
-	void OnMoveControlPoint(FPerCurveDeformer2D& InDeformer, FCurveModel& InCurveModel, TConstArrayView<int32> InVertexIndices, TConstArrayView<FVector2D> InNewControlPointPositions, EPropertyChangeType::Type InChangeType);
+	void OnMoveControlPoint(
+		FPerCurveDeformer2D& InDeformer, FCurveModel& InCurveModel, TConstArrayView<int32> InVertexIndices, TConstArrayView<FVector2D> InNewControlPointPositions,
+		EPropertyChangeType::Type InChangeType, EAxisSnapMode InSnapMode
+		);
 	/** Called when an edge is dragged. */
-	void OnMoveEdge(FPerCurveDeformer2D& InDeformer, FCurveModel& InCurveModel, TConstArrayView<int32> InVertexIndices, TConstArrayView<FVector2D> InNewControlPointPositions, EPropertyChangeType::Type InChangeType);
+	void OnMoveEdge(FPerCurveDeformer2D& InDeformer, FCurveModel& InCurveModel,
+		TConstArrayView<int32> InVertexIndices, TConstArrayView<FVector2D> InNewControlPointPositions,
+		EPropertyChangeType::Type InChangeType, EAxisSnapMode InSnapMode
+		);
 	/** Generic function for updating control points. */
-	void OnSetControlPoints(FPerCurveDeformer2D& InDeformer, FCurveModel& InCurveModel, TConstArrayView<int32> InControlPointsToUpdate, TConstArrayView<FVector2D> InNewControlPoints, EPropertyChangeType::Type InChangeType);
+	void OnSetControlPoints(
+		FPerCurveDeformer2D& InDeformer, FCurveModel& InCurveModel,
+		TConstArrayView<int32> InControlPointsToUpdate, TConstArrayView<FVector2D> InNewControlPoints,
+		EPropertyChangeType::Type InChangeType, EAxisSnapMode InSnapMode = EAxisSnapMode::Unrestricted
+		);
 	
 	/** Moves the clicked edge to the opposite edge, effectively flattening the curve sections. */
 	void MoveEdgeToOpposite(FPerCurveDeformer2D& InDeformer, FCurveModel& InCurveModel, FCurveEditor& InCurveEditor, int32 InEdgeIndex);
diff --git a/Engine/Plugins/Editor/CurveEditorTools/Source/CurveEditorTools/Private/Tools/Lattice/Drag/EAxisSnapMode.h b/Engine/Plugins/Editor/CurveEditorTools/Source/CurveEditorTools/Private/Tools/Lattice/Drag/EAxisSnapMode.h
new file mode 100644
--- /dev/null
+++ b/Engine/Plugins/Editor/CurveEditorTools/Source/CurveEditorTools/Private/Tools/Lattice/Drag/EAxisSnapMode.h
@@ -0,0 +1,24 @@
+﻿// Copyright Epic Games, Inc. All Rights Reserved.
+
+#pragma once
+
+#include "HAL/Platform.h"
+#include "Misc/EnumClassFlags.h"
+
+namespace UE::CurveEditorTools
+{
+/** Describes how keys are supposed to be snapped when using the lattice tool. */
+enum class EAxisSnapMode : uint8
+{
+	/** Mouse is unlocked, i.e can move freely in 2D space. */
+	Unrestricted,
+	
+	/** Do not snap the keys in X direction. E.g. when user has shift pressed and is moving up / down (in y direction). */
+	IgnoreSnapOnX,
+	/** Do not snap the keys in Y direction. E.g. when user has shift pressed and is moving left / right (in x direction). */
+	IgnoreSnapOnY,
+
+	IgnoreSnapping = IgnoreSnapOnY | IgnoreSnapOnX
+};
+ENUM_CLASS_FLAGS(EAxisSnapMode);
+}
\ No newline at end of file
diff --git a/Engine/Plugins/Editor/CurveEditorTools/Source/CurveEditorTools/Private/Tools/Lattice/Drag/LatticeDragOp.cpp b/Engine/Plugins/Editor/CurveEditorTools/Source/CurveEditorTools/Private/Tools/Lattice/Drag/LatticeDragOp.cpp
--- a/Engine/Plugins/Editor/CurveEditorTools/Source/CurveEditorTools/Private/Tools/Lattice/Drag/LatticeDragOp.cpp
+++ b/Engine/Plugins/Editor/CurveEditorTools/Source/CurveEditorTools/Private/Tools/Lattice/Drag/LatticeDragOp.cpp
@@ -15,7 +15,7 @@ void FLatticeDragOp::BeginDrag(const FGeometry& InGeometry, const FPointerEvent&
 {
 	InitialMousePosition = LastMousePosition = InMouseEvent.GetScreenSpacePosition();
 	AccumulateMouseMovement(InGeometry, InMouseEvent);
-	OnBeginDrag(InGeometry, InitialMousePosition);
+	OnBeginDrag(InGeometry, InitialMousePosition, AccumulatedMouseMovement->LastDragMode);
 }
 
 void FLatticeDragOp::MoveMouse(const FGeometry& InGeometry, const FPointerEvent& InMouseEvent)
@@ -27,7 +27,7 @@ void FLatticeDragOp::FinishedPointerInput()
 {
 	if (AccumulatedMouseMovement)
 	{
-		OnMoveMouse(AccumulatedMouseMovement->CachedGeometry, AccumulatedMouseMovement->AccumulatedPosition);
+		OnMoveMouse(AccumulatedMouseMovement->CachedGeometry, AccumulatedMouseMovement->AccumulatedPosition, AccumulatedMouseMovement->LastDragMode);
 		AccumulatedMouseMovement.Reset();
 	}
 }
@@ -36,7 +36,7 @@ void FLatticeDragOp::EndDrag(const FGeometry& InGeometry, const FPointerEvent& I
 {
 	// Snap the mouse position according to the axis snapping settings.
 	AccumulateMouseMovement(InGeometry, InMouseEvent);
-	OnEndDrag(InGeometry, AccumulatedMouseMovement->AccumulatedPosition);
+	OnEndDrag(InGeometry, AccumulatedMouseMovement->AccumulatedPosition, AccumulatedMouseMovement->LastDragMode);
 	AccumulatedMouseMovement.Reset();
 }
 
@@ -45,6 +45,35 @@ void FLatticeDragOp::CancelDrag()
 	OnCancelDrag();
 }
 
+namespace Private
+{
+static EAxisSnapMode DetermineSnapMode(const FCurveEditorAxisSnap& Snap, const FCurveEditorAxisSnap::FSnapState& SnapState)
+{
+	const bool bIsRestrictingWithShift = SnapState.bHasPassedThreshold;
+	if (bIsRestrictingWithShift)
+	{
+		return FMath::IsNearlyZero(SnapState.MouseLockVector.X) ? EAxisSnapMode::IgnoreSnapOnX : EAxisSnapMode::IgnoreSnapOnY;
+	}
+
+	// This handles the start of the drag: shift is pressed but not yet above the threshold.
+	// Without this, all keys would get snapped and then become unsnapped once the mouse goes over the threshold. It's a visual unpleasantness.
+	if (SnapState.bHasStartPosition)
+	{
+		return EAxisSnapMode::IgnoreSnapping;
+	}
+	
+	switch (Snap.RestrictedAxisList)
+	{
+	case ECurveEditorSnapAxis::CESA_None: return EAxisSnapMode::Unrestricted;
+	case ECurveEditorSnapAxis::CESA_X: return EAxisSnapMode::IgnoreSnapOnY;
+	case ECurveEditorSnapAxis::CESA_Y: return EAxisSnapMode::IgnoreSnapOnX;
+
+	default:
+		return EAxisSnapMode::Unrestricted;
+	}
+}
+}
+
 void FLatticeDragOp::AccumulateMouseMovement(const FGeometry& InGeometry, const FPointerEvent& InMouseEvent)
 {
 	const TSharedPtr<FCurveEditor> CurveEditorPin = CurveEditor.Pin();
@@ -54,18 +83,22 @@ void FLatticeDragOp::AccumulateMouseMovement(const FGeometry& InGeometry, const
 	}
 	
 	const FVector2D ScreenSpacePosition = InMouseEvent.GetScreenSpacePosition();
-	const FVector2D SnappedPosition = CurveEditorPin->GetAxisSnap().GetSnappedPosition(
+	LastMousePosition = ScreenSpacePosition;
+	
+	FCurveEditorAxisSnap Snap = CurveEditorPin->GetAxisSnap();
+	const FVector2D SnappedPosition = Snap.GetSnappedPosition(
 		InitialMousePosition, ScreenSpacePosition, LastMousePosition, InMouseEvent, SnapState
 		);
-	LastMousePosition = ScreenSpacePosition;
+	const EAxisSnapMode SnapMode = Private::DetermineSnapMode(Snap, SnapState);
 	
 	if (!AccumulatedMouseMovement)
 	{
-		AccumulatedMouseMovement.Emplace(InGeometry, SnappedPosition);
+		AccumulatedMouseMovement.Emplace(InGeometry, SnappedPosition, SnapMode);
 	}
 	else
 	{
 		AccumulatedMouseMovement->AccumulatedPosition = SnappedPosition;
+		AccumulatedMouseMovement->LastDragMode = SnapMode;
 	}
 }
 }
diff --git a/Engine/Plugins/Editor/CurveEditorTools/Source/CurveEditorTools/Private/Tools/Lattice/Drag/LatticeDragOp.h b/Engine/Plugins/Editor/CurveEditorTools/Source/CurveEditorTools/Private/Tools/Lattice/Drag/LatticeDragOp.h
--- a/Engine/Plugins/Editor/CurveEditorTools/Source/CurveEditorTools/Private/Tools/Lattice/Drag/LatticeDragOp.h
+++ b/Engine/Plugins/Editor/CurveEditorTools/Source/CurveEditorTools/Private/Tools/Lattice/Drag/LatticeDragOp.h
@@ -3,6 +3,7 @@
 #pragma once
 
 #include "CurveEditorSnapMetrics.h"
+#include "EAxisSnapMode.h"
 #include "Layout/Geometry.h"
 #include "Math/MathFwd.h"
 #include "Math/Vector2D.h"
@@ -30,9 +31,9 @@ public:
 
 protected:
 
-	virtual void OnBeginDrag(const FGeometry& InGeometry, const FVector2D& InInitialMousePosition) {}
-	virtual void OnMoveMouse(const FGeometry& InGeometry, const FVector2D& InScreenPosition) {}
-	virtual void OnEndDrag(const FGeometry& InGeometry, const FVector2D& InMousePosition) {}
+	virtual void OnBeginDrag(const FGeometry& InGeometry, const FVector2D& InInitialMousePosition, EAxisSnapMode InSnapMode) {}
+	virtual void OnMoveMouse(const FGeometry& InGeometry, const FVector2D& InScreenPosition, EAxisSnapMode InSnapMode) {}
+	virtual void OnEndDrag(const FGeometry& InGeometry, const FVector2D& InMousePosition, EAxisSnapMode InSnapMode) {}
 	virtual void OnCancelDrag() {}
 
 	TWeakPtr<FCurveEditor> GetCurveEditor() const { return CurveEditor; }
@@ -54,11 +55,16 @@ private:
 	struct FAccumulatedMouseMovement
 	{
 		const FGeometry CachedGeometry;
+		
+		/** The mouse position that has been accumulated by processing all pumped input this tick. */
 		FVector2D AccumulatedPosition;
+		/** The mouse lock mode that has been accumulated by processing all pumped input this tick. */
+		EAxisSnapMode LastDragMode;
 
-		explicit FAccumulatedMouseMovement(const FGeometry& InCachedGeometry, const FVector2D& InInitialPosition)
+		explicit FAccumulatedMouseMovement(const FGeometry& InCachedGeometry, const FVector2D& InInitialPosition, EAxisSnapMode InInitialDragMode)
 			: CachedGeometry(InCachedGeometry)
 			, AccumulatedPosition(InInitialPosition)
+			, LastDragMode(InInitialDragMode)
 		{}
 	};
 	TOptional<FAccumulatedMouseMovement> AccumulatedMouseMovement;
diff --git a/Engine/Plugins/Editor/CurveEditorTools/Source/CurveEditorTools/Private/Tools/Lattice/Drag/LatticeDragOp_MoveControlPoints.cpp b/Engine/Plugins/Editor/CurveEditorTools/Source/CurveEditorTools/Private/Tools/Lattice/Drag/LatticeDragOp_MoveControlPoints.cpp
--- a/Engine/Plugins/Editor/CurveEditorTools/Source/CurveEditorTools/Private/Tools/Lattice/Drag/LatticeDragOp_MoveControlPoints.cpp
+++ b/Engine/Plugins/Editor/CurveEditorTools/Source/CurveEditorTools/Private/Tools/Lattice/Drag/LatticeDragOp_MoveControlPoints.cpp
@@ -27,21 +27,21 @@ FLatticeDragOp_MoveControlPoints::FLatticeDragOp_MoveControlPoints(
 	UE_LOG(LogTemp, Warning, TEXT("Min: %s \t Max:%s"), *Bounds.MinValues.ToString(), *Bounds.MaxValues.ToString());
 }
 
-void FLatticeDragOp_MoveControlPoints::OnBeginDrag(const FGeometry& InGeometry, const FVector2D& InInitialMousePosition)
+void FLatticeDragOp_MoveControlPoints::OnBeginDrag(const FGeometry& InGeometry, const FVector2D& InInitialMousePosition, EAxisSnapMode InSnapMode)
 {
 	OnBeginDragControlPointsDelegate.ExecuteIfBound();
 }
 
-void FLatticeDragOp_MoveControlPoints::OnMoveMouse(const FGeometry& InGeometry, const FVector2D& InScreenPosition)
+void FLatticeDragOp_MoveControlPoints::OnMoveMouse(const FGeometry& InGeometry, const FVector2D& InScreenPosition, EAxisSnapMode InSnapMode)
 {
 	UpdateControlPoints(InGeometry, InScreenPosition);
-	OnInteractiveDragControlPointsDelegate.Execute(UpdatedControlPointPositions);
+	OnInteractiveDragControlPointsDelegate.Execute(UpdatedControlPointPositions, InSnapMode);
 }
 
-void FLatticeDragOp_MoveControlPoints::OnEndDrag(const FGeometry& InGeometry, const FVector2D& InMousePosition)
+void FLatticeDragOp_MoveControlPoints::OnEndDrag(const FGeometry& InGeometry, const FVector2D& InMousePosition, EAxisSnapMode InSnapMode)
 {
 	UpdateControlPoints(InGeometry, InMousePosition);
-	OnFinishDragControlPointsDelegate.Execute(UpdatedControlPointPositions);
+	OnFinishDragControlPointsDelegate.Execute(UpdatedControlPointPositions, InSnapMode);
 }
 
 void FLatticeDragOp_MoveControlPoints::UpdateControlPoints(const FGeometry& InGeometry, const FVector2D& InScreenPosition)
diff --git a/Engine/Plugins/Editor/CurveEditorTools/Source/CurveEditorTools/Private/Tools/Lattice/Drag/LatticeDragOp_MoveControlPoints.h b/Engine/Plugins/Editor/CurveEditorTools/Source/CurveEditorTools/Private/Tools/Lattice/Drag/LatticeDragOp_MoveControlPoints.h
--- a/Engine/Plugins/Editor/CurveEditorTools/Source/CurveEditorTools/Private/Tools/Lattice/Drag/LatticeDragOp_MoveControlPoints.h
+++ b/Engine/Plugins/Editor/CurveEditorTools/Source/CurveEditorTools/Private/Tools/Lattice/Drag/LatticeDragOp_MoveControlPoints.h
@@ -22,7 +22,7 @@ class FLatticeDragOp_MoveControlPoints : public FLatticeDragOp
 public:
 
 	using FOnBeginDragOperation = FSimpleDelegate;
-	DECLARE_DELEGATE_OneParam(FOnControlPointsMoved, const TConstArrayView<FVector2D>& InNewControlPointPositions);
+	DECLARE_DELEGATE_TwoParams(FOnControlPointsMoved, const TConstArrayView<FVector2D>& InNewControlPointPositions, EAxisSnapMode InSnapMode);
 
 	explicit FLatticeDragOp_MoveControlPoints(
 		TWeakPtr<FCurveEditor> InCurveEditor,
@@ -36,9 +36,9 @@ public:
 protected:
 	
 	//~ Begin FLatticeDragOp Interface
-	virtual void OnBeginDrag(const FGeometry& InGeometry, const FVector2D& InInitialMousePosition) override;
-	virtual void OnMoveMouse(const FGeometry& InGeometry, const FVector2D& InScreenPosition) override;
-	virtual void OnEndDrag(const FGeometry& InGeometry, const FVector2D& InMousePosition) override;
+	virtual void OnBeginDrag(const FGeometry& InGeometry, const FVector2D& InInitialMousePosition, EAxisSnapMode InSnapMode) override;
+	virtual void OnMoveMouse(const FGeometry& InGeometry, const FVector2D& InScreenPosition, EAxisSnapMode InSnapMode) override;
+	virtual void OnEndDrag(const FGeometry& InGeometry, const FVector2D& InMousePosition, EAxisSnapMode InSnapMode) override;
 	//~ End FLatticeDragOp Interface
 
 private:
