Lattice Tool: remove logs

#rnx
#Synced-CL 49388435

diff --git a/Engine/Plugins/Editor/CurveEditorTools/Source/CurveEditorTools/Private/Tools/Lattice/Drag/LatticeDragOp_MoveControlPoints.cpp b/Engine/Plugins/Editor/CurveEditorTools/Source/CurveEditorTools/Private/Tools/Lattice/Drag/LatticeDragOp_MoveControlPoints.cpp
--- a/Engine/Plugins/Editor/CurveEditorTools/Source/CurveEditorTools/Private/Tools/Lattice/Drag/LatticeDragOp_MoveControlPoints.cpp
+++ b/Engine/Plugins/Editor/CurveEditorTools/Source/CurveEditorTools/Private/Tools/Lattice/Drag/LatticeDragOp_MoveControlPoints.cpp
@@ -22,10 +22,7 @@ FLatticeDragOp_MoveControlPoints::FLatticeDragOp_MoveControlPoints(
 	, OnInteractiveDragControlPointsDelegate(MoveTemp(InInteractiveDragControlPointsDelegate))
 	, OnFinishDragControlPointsDelegate(MoveTemp(InFinishDragControlPointsDelegate))
 	, UpdatedControlPointPositions(InitialControlPoints)
-{
-	UE_LOG(LogTemp, Warning, TEXT("SlateMin: %s \t SlateMax:%s"), *Bounds.MinSlatePosition.ToString(), *Bounds.MaxSlatePosition.ToString());
-	UE_LOG(LogTemp, Warning, TEXT("Min: %s \t Max:%s"), *Bounds.MinValues.ToString(), *Bounds.MaxValues.ToString());
-}
+{}
 
 void FLatticeDragOp_MoveControlPoints::OnBeginDrag(const FGeometry& InGeometry, const FVector2D& InInitialMousePosition, EAxisSnapMode InSnapMode)
 {
