Curve Editor: Extract private util types from FCurveEditorLatticeTool.h into separate headers.

#Synced-CL 51881412

diff --git a/Engine/Plugins/Editor/CurveEditorTools/Source/CurveEditorTools/Private/Tools/Lattice/CurveEditorLatticeTool.cpp b/Engine/Plugins/Editor/CurveEditorTools/Source/CurveEditorTools/Private/Tools/Lattice/CurveEditorLatticeTool.cpp
--- a/Engine/Plugins/Editor/CurveEditorTools/Source/CurveEditorTools/Private/Tools/Lattice/CurveEditorLatticeTool.cpp
+++ b/Engine/Plugins/Editor/CurveEditorTools/Source/CurveEditorTools/Private/Tools/Lattice/CurveEditorLatticeTool.cpp
@@ -291,29 +291,6 @@ void FCurveEditorLatticeTool::AddReferencedObjects(FReferenceCollector& Collecto
 	Collector.AddReferencedObject(UndoObject);
 }
 
-FCurveEditorLatticeTool::FLatticeDeformerState::FLatticeDeformerState(const FLatticeBounds& Lattice, const TSharedRef<FCurveEditor>& InCurveEditor)
-	: ControlPointToCurveSpace(
-		TransformRectBetweenSpaces(Lattice.MinValues, Lattice.MaxValues, Lattice.MinValuesCurveSpace, Lattice.MaxValuesCurveSpace)
-		)
-	, GlobalDeformer(1, 1, Lattice.MinValues, Lattice.MaxValues)
-	, PerCurveData(BuildPerLatticeData(Lattice, *InCurveEditor))
-	, CurveChangeListener(CurveEditor::FCurveChangeListener::MakeForAllCurves(InCurveEditor))
-	, PanelRebuildListener(InCurveEditor)
-{}
-	
-TArray<FVector2D> FCurveEditorLatticeTool::FLatticeDeformerState::TransformControlPointsToCurveSpace() const
-{
-	const TConstArrayView<FVector2D> ControlPoints = GlobalDeformer.GetControlPoints();
-	TArray<FVector2D> Result;
-	Result.Reserve(ControlPoints.Num());
-	Algo::Transform(ControlPoints, Result, [this](const FVector2D& ControlPoint)
-	{
-		const FVector2D Transformed = ControlPointToCurveSpace.TransformPoint(ControlPoint);
-		return Transformed;
-	});
-	return Result;
-}
-
 void FCurveEditorLatticeTool::UpdateDeformerBounds(const TSharedRef<FCurveEditor>& InCurveEditor)
 {
 	if (LatticeDeformerState.IsSet())
@@ -494,7 +471,7 @@ void FCurveEditorLatticeTool::StartDragDrop(const FGeometry& InViewContainerGeom
 	}
 }
 	
-TOptional<FCurveEditorLatticeTool::FDragOpData> FCurveEditorLatticeTool::CreateDragDrop(
+TOptional<FDragOpData> FCurveEditorLatticeTool::CreateDragDrop(
 	const FLatticeBounds& InBounds, const FLatticeHoverState& InHoverState
 	)
 {
@@ -514,7 +491,7 @@ TOptional<FCurveEditorLatticeTool::FDragOpData> FCurveEditorLatticeTool::CreateD
 	return {};
 }
 
-TOptional<FCurveEditorLatticeTool::FDragOpData> FCurveEditorLatticeTool::CreateDragDrop_Point(
+TOptional<FDragOpData> FCurveEditorLatticeTool::CreateDragDrop_Point(
 	const FLatticeBounds& InBounds, const int32 InControlPointIndex
 	)
 {
@@ -566,7 +543,7 @@ TOptional<FCurveEditorLatticeTool::FDragOpData> FCurveEditorLatticeTool::CreateD
 	return DragDrop;
 }
 
-TOptional<FCurveEditorLatticeTool::FDragOpData> FCurveEditorLatticeTool::CreateDragDrop_Edge(
+TOptional<FDragOpData> FCurveEditorLatticeTool::CreateDragDrop_Edge(
 	const FLatticeBounds& InBounds, int32 InEdgeIndex
 	)
 {
@@ -619,7 +596,7 @@ TOptional<FCurveEditorLatticeTool::FDragOpData> FCurveEditorLatticeTool::CreateD
 	return DragDrop;
 }
 
-TOptional<FCurveEditorLatticeTool::FDragOpData> FCurveEditorLatticeTool::CreateDragDrop_Cell(
+TOptional<FDragOpData> FCurveEditorLatticeTool::CreateDragDrop_Cell(
 	const FLatticeBounds& InBounds, int32 InCellIndex
 	)
 {
diff --git a/Engine/Plugins/Editor/CurveEditorTools/Source/CurveEditorTools/Private/Tools/Lattice/CurveEditorLatticeTool.h b/Engine/Plugins/Editor/CurveEditorTools/Source/CurveEditorTools/Private/Tools/Lattice/CurveEditorLatticeTool.h
--- a/Engine/Plugins/Editor/CurveEditorTools/Source/CurveEditorTools/Private/Tools/Lattice/CurveEditorLatticeTool.h
+++ b/Engine/Plugins/Editor/CurveEditorTools/Source/CurveEditorTools/Private/Tools/Lattice/CurveEditorLatticeTool.h
@@ -7,6 +7,7 @@
 #include "ICurveEditorToolExtension.h"
 #include "LatticeControlsDrawData.h"
 #include "LatticeDeformer2D.h"
+#include "LatticeDeformerState.h"
 #include "LatticeFwd.h"
 #include "Mirror/LatticeEdgeTangentMirrorOp.h"
 #include "Mirror/LatticePointTangentMirrorOp.h"
@@ -16,6 +17,7 @@
 #include "Misc/LatticeDrawUtils.h"
 #include "PerCurveLatticeData.h"
 #include "ScopedTransaction.h"
+#include "Drag/DelayedDragData.h"
 #include "Modification/Utils/ScopedCurveChange.h"
 #include "Templates/SharedPointer.h"
 #include "UObject/UnrealType.h"
@@ -30,23 +32,6 @@ namespace UE::CurveEditorTools
 {
 struct FLatticeBounds;
 
-enum class ELatticeUpdateFlags : uint8
-{
-	None,
-	
-	/**
-	 * Tries to restore the grid bounds from what is saved in the undo object.
-	 * This is set when an undo operation is performed.
-	 *
-	 * If the data is compatible, the grid is restored to that and the grid is not recomputed.
-	 * If the data is incompatible, the entire deformer is recomputed.
-	 *
-	 * The data may be incompatible, if e.g. you switched from normalized view to absolute view.
-	 */
-	TryRestoreBoundsFromUndo = 1 << 1,
-};
-ENUM_CLASS_FLAGS(ELatticeUpdateFlags);
-	
 /**
  * Places a lattice deformer over the selection.
  * 
@@ -89,48 +74,7 @@ private:
 	
 	/** Weak pointer back to the owning curve editor. */
 	TWeakPtr<FCurveEditor> WeakCurveEditor;
-
-	/** Holds any data that exists only when the deformer exists. The deformer exists the selection consists of least 2 keys at different locations. */
-	struct FLatticeDeformerState
-	{
-		/**
-		 * Transforms the absolute Deformer.GetControlPoints(), which are just the FKeyPosition values, to curve space values of the current view.
-		 *
-		 * Reminder: CurveSpace are axis values that the SCurveEditorView displays on its axis.
-		 * In Absolute mode, that's the literal FKeyPosition values
-		 * In Normalized mode, Y is in 0 to 1 range (and X continues to be FKeyPosition::InputValue).
-		 */
-		FTransform2d ControlPointToCurveSpace;
-
-		/**
-		 * Empty "deformer" that exists purely for the UI. Its control points, edges, etc. are displayed in the view.
-		 * 
-		 * Each curve has its own lattice deformer that actually moves that curves keys.
-		 * The per curve deformers control points are the result of applying the per-curve transform to Deformer.
-		 * This is needed to support all view modes (Absolute, Normalized, etc.)
-		 *
-		 * The Deformer's control points are always in absolute key space (i.e. exactly the values for FKeyPosition::InputValue and OutputValue).
-		 */
-		FGlobalLatticeDeformer2D GlobalDeformer;
-		TMap<FCurveModelID, FPerCurveLatticeData> PerCurveData; 
-
-		/** Tells us when a curve is externally modified, so we can update the bounds. */
-		CurveEditor::FCurveChangeListener CurveChangeListener;
-		/** Tells us when the panel is regenerated, in which case we must regenerate the lattice shape (e.g. could go from absolute -> normalized mode). */
-		FCurveViewRebuildListener PanelRebuildListener;
-		
-		/** Guard to not recompute bounds when FCurveModel::OnCurveModified broadcasts due to a change we have initiated. */
-		bool bIsModifyingCurves = false;
-		/** If true, the curves have been modified and the lattice overlay should be recomputed next frame. */
-		bool bHasRequestedRefresh = false;
-		/** The flags to use when bHasRequestedRefresh == true. */
-		ELatticeUpdateFlags RefreshFlags = ELatticeUpdateFlags::None;
-
-		explicit FLatticeDeformerState(const FLatticeBounds& Lattice, const TSharedRef<FCurveEditor>& InCurveEditor);
-
-		/** @return Control points transformed to SCurveEditorView's curve space. */
-		TArray<FVector2D> TransformControlPointsToCurveSpace() const;
-	};
+	
 	/** Set when there is a valid selection in view. */
 	TOptional<FLatticeDeformerState> LatticeDeformerState;
 	
@@ -139,61 +83,7 @@ private:
 
 	/** The last mouse position received by OnMouseMove. */
 	FVector2D FinalMouseScreenSpacePosition;
-
-	/** Data used after FDelayedDragData has detected a drag. The user is now actively dragging a control point, edge, or cell. */
-	struct FDragOpData
-	{
-		/** Handles drags that started on a control point. */
-		TUniquePtr<FLatticeDragOp> DragOp;
-		/** This transaction is active for the duration of the drag. */
-		TUniquePtr<FScopedTransaction> Transaction;
-		/** Transactions the key changes */
-		TUniquePtr<CurveEditor::FScopedCurveChange> KeyChange;
-		
-		/** Adjust tangents while dragging on an edge.*/
-		TMap<FCurveModelID, FLatticeEdgeTangentsMirrorOp> EdgeTangentMirroringOps;
-		/** Adjusts tangents while dragging on a control point. */
-		TMap<FCurveModelID, FLatticePointTangentsMirrorOp> PointTangentMirroringOps;
-
-		/** False until the first drag actually changes keys. We only need to capture undo state just before the first change. */
-		bool bSavedUndoState = false;
-
-		explicit FDragOpData(TWeakPtr<FCurveEditor> InCurveEditor, TUniquePtr<FLatticeDragOp> DragOp, const FText& TransactionText)
-			: DragOp(MoveTemp(DragOp))
-			, Transaction(MakeUnique<FScopedTransaction>(TEXT("CurveEditorLatticeTool"), TransactionText, nullptr))
-			, KeyChange([InCurveEditor]
-			{
-				using namespace UE::CurveEditor;
-				return MakeUnique<FScopedCurveChange>(
-					FCurvesSnapshotBuilder(
-						InCurveEditor,
-						ECurveChangeFlags::MoveKeysAndRemoveStackedKeys | ECurveChangeFlags::KeyAttributes
-						).TrackSelectedCurves()
-					);
-			}())
-		{}
-	};
-
-	/** Data used while the user is holding down the left mouse button. Once the mouse has moved enough, we "detect" a drag. */
-	struct FDelayedDragData
-	{
-		/** Set when attempting to move a drag handle. This allows us to tell the difference between a click and a click-drag. */
-		FDelayedDrag DelayedDrag;
-
-		/** Bounds when the op was started. */
-		const FLatticeBounds Bounds;
-		/** Hover state when the op was started. */
-		const FLatticeHoverState HoverState;
-
-		/** Set if the drag up was started. */
-		TOptional<FDragOpData> ActiveOperation;
-		
-		explicit FDelayedDragData(
-			const FVector2D& InInitialPosition, const FKey& InEffectiveKey, const FLatticeBounds& InBounds, const FLatticeHoverState& InHoverState
-			)
-			: DelayedDrag(InInitialPosition, InEffectiveKey), Bounds(InBounds), HoverState(InHoverState)
-		{}
-	};
+	/** Data used to detect dragging. */
 	TOptional<FDelayedDragData> DelayedDragData;
 
 	/**
diff --git a/Engine/Plugins/Editor/CurveEditorTools/Source/CurveEditorTools/Private/Tools/Lattice/Drag/DelayedDragData.h b/Engine/Plugins/Editor/CurveEditorTools/Source/CurveEditorTools/Private/Tools/Lattice/Drag/DelayedDragData.h
new file mode 100644
--- /dev/null
+++ b/Engine/Plugins/Editor/CurveEditorTools/Source/CurveEditorTools/Private/Tools/Lattice/Drag/DelayedDragData.h
@@ -0,0 +1,33 @@
+﻿// Copyright Epic Games, Inc. All Rights Reserved.
+
+#pragma once
+
+#include "DragOpData.h"
+#include "Framework/DelayedDrag.h"
+#include "Misc/Optional.h"
+#include "Tools/Lattice/LatticeControlsDrawData.h"
+#include "Tools/Lattice/Misc/LatticeDrawUtils.h"
+
+namespace UE::CurveEditorTools
+{
+/** Data used while the user is holding down the left mouse button. Once the mouse has moved enough, we "detect" a drag. */
+struct FDelayedDragData
+{
+	/** Set when attempting to move a drag handle. This allows us to tell the difference between a click and a click-drag. */
+	FDelayedDrag DelayedDrag;
+
+	/** Bounds when the op was started. */
+	const FLatticeBounds Bounds;
+	/** Hover state when the op was started. */
+	const FLatticeHoverState HoverState;
+
+	/** Set if the drag up was started. */
+	TOptional<FDragOpData> ActiveOperation;
+		
+	explicit FDelayedDragData(
+		const FVector2D& InInitialPosition, const FKey& InEffectiveKey, const FLatticeBounds& InBounds, const FLatticeHoverState& InHoverState
+		)
+		: DelayedDrag(InInitialPosition, InEffectiveKey), Bounds(InBounds), HoverState(InHoverState)
+	{}
+};
+}
diff --git a/Engine/Plugins/Editor/CurveEditorTools/Source/CurveEditorTools/Private/Tools/Lattice/Drag/DragOpData.h b/Engine/Plugins/Editor/CurveEditorTools/Source/CurveEditorTools/Private/Tools/Lattice/Drag/DragOpData.h
new file mode 100644
--- /dev/null
+++ b/Engine/Plugins/Editor/CurveEditorTools/Source/CurveEditorTools/Private/Tools/Lattice/Drag/DragOpData.h
@@ -0,0 +1,47 @@
+﻿// Copyright Epic Games, Inc. All Rights Reserved.
+
+#pragma once
+
+#include "LatticeDragOp.h"
+#include "Modification/Utils/ScopedCurveChange.h"
+#include "ScopedTransaction.h"
+#include "Templates/UniquePtr.h"
+#include "Tools/Lattice/Mirror/LatticeEdgeTangentMirrorOp.h"
+#include "Tools/Lattice/Mirror/LatticePointTangentMirrorOp.h"
+
+namespace UE::CurveEditorTools
+{
+/** Data used after FDelayedDragData has detected a drag. The user is now actively dragging a control point, edge, or cell. */
+struct FDragOpData
+{
+	/** Handles drags that started on a control point. */
+	TUniquePtr<FLatticeDragOp> DragOp;
+	/** This transaction is active for the duration of the drag. */
+	TUniquePtr<FScopedTransaction> Transaction;
+	/** Transactions the key changes */
+	TUniquePtr<CurveEditor::FScopedCurveChange> KeyChange;
+		
+	/** Adjust tangents while dragging on an edge.*/
+	TMap<FCurveModelID, FLatticeEdgeTangentsMirrorOp> EdgeTangentMirroringOps;
+	/** Adjusts tangents while dragging on a control point. */
+	TMap<FCurveModelID, FLatticePointTangentsMirrorOp> PointTangentMirroringOps;
+
+	/** False until the first drag actually changes keys. We only need to capture undo state just before the first change. */
+	bool bSavedUndoState = false;
+
+	explicit FDragOpData(TWeakPtr<FCurveEditor> InCurveEditor, TUniquePtr<FLatticeDragOp> DragOp, const FText& TransactionText)
+		: DragOp(MoveTemp(DragOp))
+		, Transaction(MakeUnique<FScopedTransaction>(TEXT("CurveEditorLatticeTool"), TransactionText, nullptr))
+		, KeyChange([InCurveEditor]
+		{
+			using namespace UE::CurveEditor;
+			return MakeUnique<FScopedCurveChange>(
+				FCurvesSnapshotBuilder(
+					InCurveEditor,
+					ECurveChangeFlags::MoveKeysAndRemoveStackedKeys | ECurveChangeFlags::KeyAttributes
+					).TrackSelectedCurves()
+				);
+		}())
+	{}
+};
+}
diff --git a/Engine/Plugins/Editor/CurveEditorTools/Source/CurveEditorTools/Private/Tools/Lattice/LatticeDeformerState.cpp b/Engine/Plugins/Editor/CurveEditorTools/Source/CurveEditorTools/Private/Tools/Lattice/LatticeDeformerState.cpp
new file mode 100644
--- /dev/null
+++ b/Engine/Plugins/Editor/CurveEditorTools/Source/CurveEditorTools/Private/Tools/Lattice/LatticeDeformerState.cpp
@@ -0,0 +1,34 @@
+﻿// Copyright Epic Games, Inc. All Rights Reserved.
+
+
+#include "LatticeDeformerState.h"
+
+#include "Algo/Transform.h"
+#include "Misc/LatticeDrawUtils.h"
+#include "Misc/VectorMathUtils.h"
+
+namespace UE::CurveEditorTools
+{
+FLatticeDeformerState::FLatticeDeformerState(const FLatticeBounds& Lattice, const TSharedRef<FCurveEditor>& InCurveEditor)
+	: ControlPointToCurveSpace(
+		TransformRectBetweenSpaces(Lattice.MinValues, Lattice.MaxValues, Lattice.MinValuesCurveSpace, Lattice.MaxValuesCurveSpace)
+		)
+	, GlobalDeformer(1, 1, Lattice.MinValues, Lattice.MaxValues)
+	, PerCurveData(BuildPerLatticeData(Lattice, *InCurveEditor))
+	, CurveChangeListener(CurveEditor::FCurveChangeListener::MakeForAllCurves(InCurveEditor))
+	, PanelRebuildListener(InCurveEditor)
+{}
+	
+TArray<FVector2D> FLatticeDeformerState::TransformControlPointsToCurveSpace() const
+{
+	const TConstArrayView<FVector2D> ControlPoints = GlobalDeformer.GetControlPoints();
+	TArray<FVector2D> Result;
+	Result.Reserve(ControlPoints.Num());
+	Algo::Transform(ControlPoints, Result, [this](const FVector2D& ControlPoint)
+	{
+		const FVector2D Transformed = ControlPointToCurveSpace.TransformPoint(ControlPoint);
+		return Transformed;
+	});
+	return Result;
+}
+}
diff --git a/Engine/Plugins/Editor/CurveEditorTools/Source/CurveEditorTools/Private/Tools/Lattice/LatticeDeformerState.h b/Engine/Plugins/Editor/CurveEditorTools/Source/CurveEditorTools/Private/Tools/Lattice/LatticeDeformerState.h
new file mode 100644
--- /dev/null
+++ b/Engine/Plugins/Editor/CurveEditorTools/Source/CurveEditorTools/Private/Tools/Lattice/LatticeDeformerState.h
@@ -0,0 +1,75 @@
+﻿// Copyright Epic Games, Inc. All Rights Reserved.
+
+#pragma once
+
+#include "Containers/Array.h"
+#include "Containers/Map.h"
+#include "CurveEditorTypes.h"
+#include "LatticeFwd.h"
+#include "Math/TransformCalculus2D.h"
+#include "Misc/CurveChangeListener.h"
+#include "Misc/CurveViewRebuildListener.h"
+#include "Misc/EnumClassFlags.h"
+#include "PerCurveLatticeData.h"
+
+namespace UE::CurveEditorTools
+{
+enum class ELatticeUpdateFlags : uint8
+{
+	None,
+	
+	/**
+	 * Tries to restore the grid bounds from what is saved in the undo object.
+	 * This is set when an undo operation is performed.
+	 *
+	 * If the data is compatible, the grid is restored to that and the grid is not recomputed.
+	 * If the data is incompatible, the entire deformer is recomputed.
+	 *
+	 * The data may be incompatible, if e.g. you switched from normalized view to absolute view.
+	 */
+	TryRestoreBoundsFromUndo = 1 << 1,
+};
+ENUM_CLASS_FLAGS(ELatticeUpdateFlags);
+
+/** Holds any data that exists only when the deformer exists. The deformer exists the selection consists of least 2 keys at different locations. */
+struct FLatticeDeformerState
+{
+	/**
+	 * Transforms the absolute Deformer.GetControlPoints(), which are just the FKeyPosition values, to curve space values of the current view.
+	 *
+	 * Reminder: CurveSpace are axis values that the SCurveEditorView displays on its axis.
+	 * In Absolute mode, that's the literal FKeyPosition values
+	 * In Normalized mode, Y is in 0 to 1 range (and X continues to be FKeyPosition::InputValue).
+	 */
+	FTransform2d ControlPointToCurveSpace;
+
+	/**
+	 * Empty "deformer" that exists purely for the UI. Its control points, edges, etc. are displayed in the view.
+	 * 
+	 * Each curve has its own lattice deformer that actually moves that curves keys.
+	 * The per curve deformers control points are the result of applying the per-curve transform to Deformer.
+	 * This is needed to support all view modes (Absolute, Normalized, etc.)
+	 *
+	 * The Deformer's control points are always in absolute key space (i.e. exactly the values for FKeyPosition::InputValue and OutputValue).
+	 */
+	FGlobalLatticeDeformer2D GlobalDeformer;
+	TMap<FCurveModelID, FPerCurveLatticeData> PerCurveData; 
+
+	/** Tells us when a curve is externally modified, so we can update the bounds. */
+	CurveEditor::FCurveChangeListener CurveChangeListener;
+	/** Tells us when the panel is regenerated, in which case we must regenerate the lattice shape (e.g. could go from absolute -> normalized mode). */
+	FCurveViewRebuildListener PanelRebuildListener;
+	
+	/** Guard to not recompute bounds when FCurveModel::OnCurveModified broadcasts due to a change we have initiated. */
+	bool bIsModifyingCurves = false;
+	/** If true, the curves have been modified and the lattice overlay should be recomputed next frame. */
+	bool bHasRequestedRefresh = false;
+	/** The flags to use when bHasRequestedRefresh == true. */
+	ELatticeUpdateFlags RefreshFlags = ELatticeUpdateFlags::None;
+
+	explicit FLatticeDeformerState(const FLatticeBounds& Lattice, const TSharedRef<FCurveEditor>& InCurveEditor);
+
+	/** @return Control points transformed to SCurveEditorView's curve space. */
+	TArray<FVector2D> TransformControlPointsToCurveSpace() const;
+};
+}
