Curve Editor: Fix lattice tool flipping bool values on bool curves, like IK switches, when only moving in X-direction, i.e. shift-dragging.

On bool curves the old values would be e.g. 0 but the lattice tool would compute 0.1, which internally the FCurveModel implementation interpreted as 0.1 != 0 == true: this caused the key values to flip.

The core change: when only moving in X direction, set the Y values the points originally had instead of the values computed by lattice to avoid tiny double rounding errors.

Other misc changes this logic adjustment caused:
- The original value is stored in the new FLatticePointMetaData, which required several misc signature adjustments.
- Renamed EKeySnapMode -> EKeyMoveMode

#Synced-CL 51881412

#Synced-CL 51881412

diff --git a/Engine/Plugins/Editor/CurveEditorTools/Source/CurveEditorTools/Private/Misc/CurvePointSnapper.cpp b/Engine/Plugins/Editor/CurveEditorTools/Source/CurveEditorTools/Private/Misc/CurvePointSnapper.cpp
--- a/Engine/Plugins/Editor/CurveEditorTools/Source/CurveEditorTools/Private/Misc/CurvePointSnapper.cpp
+++ b/Engine/Plugins/Editor/CurveEditorTools/Source/CurveEditorTools/Private/Misc/CurvePointSnapper.cpp
@@ -25,19 +25,19 @@ TOptional<FCurvePointSnapper> FCurvePointSnapper::MakeSnapper(FCurveEditor& InCu
 	return FCurvePointSnapper(View->GetCurveSpace(InCurveModelId), InCurveEditor.GetCurveSnapMetrics(InCurveModelId));
 }
 
-FVector2D FCurvePointSnapper::SnapPoint(const FVector2D& InCurveSpacePoint, EAxisSnapMode InDragMode) const
+FVector2D FCurvePointSnapper::SnapPoint(const FVector2D& InCurveSpacePoint, EKeyMoveMode InDragMode) const
 {
 	// LockedX & LockedY mean that the user is e.g. shift dragging: this means they're intentionally moving in only one direction.
 	return FVector2D
 	{
 		// If moving only up & down, then it would be weird if the keys snapped in x direction.
-		EnumHasAnyFlags(InDragMode, EAxisSnapMode::IgnoreSnapOnX) ? InCurveSpacePoint.X : SnapMetrics.SnapInputSeconds(InCurveSpacePoint.X),
+		EnumHasAnyFlags(InDragMode, EKeyMoveMode::IgnoreX) ? InCurveSpacePoint.X : SnapMetrics.SnapInputSeconds(InCurveSpacePoint.X),
 		// If moving only left & right, then it would be weird if the keys snapped in y direction.
-		EnumHasAnyFlags(InDragMode, EAxisSnapMode::IgnoreSnapOnY) ? InCurveSpacePoint.Y : SnapMetrics.SnapOutput(InCurveSpacePoint.Y)
+		EnumHasAnyFlags(InDragMode, EKeyMoveMode::IgnoreY) ? InCurveSpacePoint.Y : SnapMetrics.SnapOutput(InCurveSpacePoint.Y)
 	};
 }
 
-FKeyPosition FCurvePointSnapper::SnapKey(const FKeyPosition& InKeyPosition, EAxisSnapMode InDragMode) const
+FKeyPosition FCurvePointSnapper::SnapKey(const FKeyPosition& InKeyPosition, EKeyMoveMode InDragMode) const
 {
 	const FVector2D Point = SnapPoint({ InKeyPosition.InputValue, InKeyPosition.OutputValue }, InDragMode);
 	return { Point.X, Point.Y };
diff --git a/Engine/Plugins/Editor/CurveEditorTools/Source/CurveEditorTools/Private/Misc/CurvePointSnapper.h b/Engine/Plugins/Editor/CurveEditorTools/Source/CurveEditorTools/Private/Misc/CurvePointSnapper.h
--- a/Engine/Plugins/Editor/CurveEditorTools/Source/CurveEditorTools/Private/Misc/CurvePointSnapper.h
+++ b/Engine/Plugins/Editor/CurveEditorTools/Source/CurveEditorTools/Private/Misc/CurvePointSnapper.h
@@ -4,7 +4,7 @@
 
 #include "CurveEditorScreenSpace.h"
 #include "CurveEditorSnapMetrics.h"
-#include "Tools/Lattice/Drag/EAxisSnapMode.h"
+#include "Tools/Lattice/Drag/EKeyMoveMode.h"
 
 struct FCurveModelID;
 struct FKeyPosition;
@@ -22,10 +22,10 @@ public:
 	static TOptional<FCurvePointSnapper> MakeSnapper(FCurveEditor& InCurveEditor, const FCurveModelID& InCurveModelId);
 
 	/** Snaps the point as per the snapping settings */
-	FVector2D SnapPoint(const FVector2D& InCurveSpacePoint, EAxisSnapMode InDragMode = EAxisSnapMode::Unrestricted) const;
+	FVector2D SnapPoint(const FVector2D& InCurveSpacePoint, EKeyMoveMode InDragMode = EKeyMoveMode::Unrestricted) const;
 
 	/** Snaps the key as per the snapping settings */
-	FKeyPosition SnapKey(const FKeyPosition& InKeyPosition, EAxisSnapMode InDragMode = EAxisSnapMode::Unrestricted) const;
+	FKeyPosition SnapKey(const FKeyPosition& InKeyPosition, EKeyMoveMode InDragMode = EKeyMoveMode::Unrestricted) const;
 
 private:
 
diff --git a/Engine/Plugins/Editor/CurveEditorTools/Source/CurveEditorTools/Private/Tools/Lattice/CurveEditorLatticeTool.cpp b/Engine/Plugins/Editor/CurveEditorTools/Source/CurveEditorTools/Private/Tools/Lattice/CurveEditorLatticeTool.cpp
--- a/Engine/Plugins/Editor/CurveEditorTools/Source/CurveEditorTools/Private/Tools/Lattice/CurveEditorLatticeTool.cpp
+++ b/Engine/Plugins/Editor/CurveEditorTools/Source/CurveEditorTools/Private/Tools/Lattice/CurveEditorLatticeTool.cpp
@@ -7,6 +7,7 @@
 #include "CurveEditorToolCommands.h"
 #include "Drag/LatticeDragOp_MoveControlPoints.h"
 #include "Editor.h"
+#include "LatticePointMetaData.h"
 #include "Editor/TransBuffer.h"
 #include "Internationalization/Text.h"
 #include "Mirror/LatticeEdgeTangentMirrorOp.h"
@@ -501,7 +502,7 @@ TOptional<FDragOpData> FCurveEditorLatticeTool::CreateDragDrop_Point(
 	const auto MakeDelegate = [this, InControlPointIndex](EPropertyChangeType::Type ChangeType)
 	{
 		return FHandleMovedPoints::CreateLambda(
-			[this, InControlPointIndex, ChangeType](const TConstArrayView<FVector2D>& InNewControlPointPositions, EAxisSnapMode InSnapMode)
+			[this, InControlPointIndex, ChangeType](const TConstArrayView<FVector2D>& InNewControlPointPositions, EKeyMoveMode InSnapMode)
 			{
 				const int32 Indices[1] { InControlPointIndex };
 				HandleDragPoints(&FCurveEditorLatticeTool::OnMoveControlPoint,Indices, InNewControlPointPositions, ChangeType, InSnapMode);
@@ -553,7 +554,7 @@ TOptional<FDragOpData> FCurveEditorLatticeTool::CreateDragDrop_Edge(
 	const auto MakeDelegate = [this, EdgeIndices](EPropertyChangeType::Type ChangeType)
 	{
 		return FHandleMovedPoints::CreateLambda(
-			[this, EdgeIndices, ChangeType](const TConstArrayView<FVector2D>& InNewControlPointPositions, EAxisSnapMode InSnapMode)
+			[this, EdgeIndices, ChangeType](const TConstArrayView<FVector2D>& InNewControlPointPositions, EKeyMoveMode InSnapMode)
 			{
 				const int32 Indices[2]{ EdgeIndices.Start, EdgeIndices.End };
 				HandleDragPoints(&FCurveEditorLatticeTool::OnMoveEdge, Indices, InNewControlPointPositions, ChangeType, InSnapMode);
@@ -607,7 +608,7 @@ TOptional<FDragOpData> FCurveEditorLatticeTool::CreateDragDrop_Cell(
 	const auto MakeDelegate = [this, CellIndices](EPropertyChangeType::Type ChangeType)
 	{
 		return FHandleMovedPoints::CreateLambda(
-			[this, CellIndices, ChangeType](const TConstArrayView<FVector2D>& InNewControlPointPositions, EAxisSnapMode InSnapMode)
+			[this, CellIndices, ChangeType](const TConstArrayView<FVector2D>& InNewControlPointPositions, EKeyMoveMode InSnapMode)
 			{
 				const int32 Indices[4]{ CellIndices.TopLeft, CellIndices.TopRight, CellIndices.BottomRight, CellIndices.BottomLeft };
 				HandleDragPoints(&FCurveEditorLatticeTool::OnSetControlPoints, Indices, InNewControlPointPositions, ChangeType, InSnapMode);
@@ -640,7 +641,7 @@ TOptional<FDragOpData> FCurveEditorLatticeTool::CreateDragDrop_Cell(
 void FCurveEditorLatticeTool::HandleDragPoints(
 	FMovePointsSignature InMovePointsCallback,
 	TConstArrayView<int32> InVertexIndices, TConstArrayView<FVector2D> InNewControlPoints,
-	EPropertyChangeType::Type InChangeType, EAxisSnapMode InSnapMode
+	EPropertyChangeType::Type InChangeType, EKeyMoveMode InSnapMode
 	)
 {
 	const TSharedPtr<FCurveEditor> CurveEditorPin = WeakCurveEditor.Pin();
@@ -722,7 +723,7 @@ void FCurveEditorLatticeTool::HandleMirrorGrid()
 void FCurveEditorLatticeTool::OnMoveControlPoint(
 	FPerCurveDeformer2D& InDeformer, FCurveModel& InCurveModel,
 	TConstArrayView<int32> InVertexIndices, TConstArrayView<FVector2D> InNewControlPointPositions,
-	EPropertyChangeType::Type InChangeType, EAxisSnapMode InSnapMode
+	EPropertyChangeType::Type InChangeType, EKeyMoveMode InSnapMode
 	)
 {
 	check(InVertexIndices.Num() == 1 && InNewControlPointPositions.Num() == 1);
@@ -740,7 +741,7 @@ void FCurveEditorLatticeTool::OnMoveControlPoint(
 void FCurveEditorLatticeTool::OnMoveEdge(
 	FPerCurveDeformer2D& InDeformer, FCurveModel& InCurveModel,
 	TConstArrayView<int32> InVertexIndices, TConstArrayView<FVector2D> InNewControlPointPositions,
-	EPropertyChangeType::Type InChangeType, EAxisSnapMode InSnapMode
+	EPropertyChangeType::Type InChangeType, EKeyMoveMode InSnapMode
 	)
 {
 	check(InVertexIndices.Num() == 2 && InNewControlPointPositions.Num() == 2);
@@ -758,7 +759,7 @@ void FCurveEditorLatticeTool::OnMoveEdge(
 void FCurveEditorLatticeTool::OnSetControlPoints(
 	FPerCurveDeformer2D& InDeformer, FCurveModel& InCurveModel,
 	TConstArrayView<int32> InControlPointsToUpdate, TConstArrayView<FVector2D> InNewControlPoints,
-	EPropertyChangeType::Type InChangeType, EAxisSnapMode InSnapMode
+	EPropertyChangeType::Type InChangeType, EKeyMoveMode InSnapMode
 	)
 {
 	ensure(LatticeDeformerState->bIsModifyingCurves);
@@ -781,10 +782,16 @@ void FCurveEditorLatticeTool::OnSetControlPoints(
 	Keys.Reserve(NumKeys);
 	
 	InDeformer.UpdateControlPoints(InControlPointsToUpdate, InNewControlPoints,
-		[InSnapMode, &Snapper, &Handles, &Keys](const FKeyHandle& KeyHandle, const FVector2D& KeyPosition)
+		[InSnapMode, &Snapper, &Handles, &Keys](const FLatticePointMetaData& MetaData, const FVector2D& KeyPosition)
 		{
-			Handles.Add(KeyHandle);
-			Keys.Add(Snapper->SnapKey(FKeyPosition{ KeyPosition.X, KeyPosition.Y }, InSnapMode));
+			Handles.Add(MetaData.KeyHandle);
+			const FKeyPosition& Original = MetaData.OriginalPosition;
+			const FKeyPosition Position
+			{
+				EnumHasAnyFlags(InSnapMode, EKeyMoveMode::IgnoreX) ? Original.InputValue : KeyPosition.X,
+				EnumHasAnyFlags(InSnapMode, EKeyMoveMode::IgnoreY) ? Original.OutputValue : KeyPosition.Y,
+			};
+			Keys.Add(Snapper->SnapKey(Position, InSnapMode));
 		});
 
 	// Some stats all with 9000 keys. Varying degree of changed keys, i.e. how far you move your mouse in quick succession.
@@ -832,7 +839,10 @@ void FCurveEditorLatticeTool::MirrorGrid(FPerCurveDeformer2D& InDeformer, FCurve
 		return;
 	}
 	
-	const TConstArrayView<FKeyHandle> Keys = InDeformer.GetCellMetaData(0);
+	TArray<FKeyHandle> Keys;
+	const TConstArrayView<FLatticePointMetaData> CellMetaData = InDeformer.GetCellMetaData(0);
+	Algo::Transform(CellMetaData, Keys, [](const FLatticePointMetaData& MetaData){ return MetaData.KeyHandle; });
+	
 	const auto[Min, Max] = CurveEditor::FindMinMaxHeight(Keys, InCurveModel);
 	const double Midpoint = Min + (Max - Min) * 0.5;
 	CurveEditor::MirrorTangents(Keys, *InCurveModel.GetId(), Min, Max, Midpoint, InCurveEditor);
diff --git a/Engine/Plugins/Editor/CurveEditorTools/Source/CurveEditorTools/Private/Tools/Lattice/CurveEditorLatticeTool.h b/Engine/Plugins/Editor/CurveEditorTools/Source/CurveEditorTools/Private/Tools/Lattice/CurveEditorLatticeTool.h
--- a/Engine/Plugins/Editor/CurveEditorTools/Source/CurveEditorTools/Private/Tools/Lattice/CurveEditorLatticeTool.h
+++ b/Engine/Plugins/Editor/CurveEditorTools/Source/CurveEditorTools/Private/Tools/Lattice/CurveEditorLatticeTool.h
@@ -128,12 +128,12 @@ private:
 	using FMovePointsSignature = void(FCurveEditorLatticeTool::*)(
 		FPerCurveDeformer2D& InDeformer, FCurveModel& InCurveModel,
 		TConstArrayView<int32> InVertexIndices, TConstArrayView<FVector2D> InNewControlPointPositions,
-		EPropertyChangeType::Type InChangeType, EAxisSnapMode InSnapMode
+		EPropertyChangeType::Type InChangeType, EKeyMoveMode InSnapMode
 		);
 	/** Handles dragging a control point, edge, or the lattice center. */
 	void HandleDragPoints(
 		FMovePointsSignature InMovePointsCallback, TConstArrayView<int32> InVertexIndices, TConstArrayView<FVector2D> InNewControlPoints,
-		EPropertyChangeType::Type InChangeType, EAxisSnapMode InSnapMode
+		EPropertyChangeType::Type InChangeType, EKeyMoveMode InSnapMode
 		);
 	/** Handles double-clicking an edge. */
 	void HandleMoveEdgeToOpposite(int32 InEdgeIndex);
@@ -143,18 +143,18 @@ private:
 	/** Called when a control point is dragged. */
 	void OnMoveControlPoint(
 		FPerCurveDeformer2D& InDeformer, FCurveModel& InCurveModel, TConstArrayView<int32> InVertexIndices, TConstArrayView<FVector2D> InNewControlPointPositions,
-		EPropertyChangeType::Type InChangeType, EAxisSnapMode InSnapMode
+		EPropertyChangeType::Type InChangeType, EKeyMoveMode InSnapMode
 		);
 	/** Called when an edge is dragged. */
 	void OnMoveEdge(FPerCurveDeformer2D& InDeformer, FCurveModel& InCurveModel,
 		TConstArrayView<int32> InVertexIndices, TConstArrayView<FVector2D> InNewControlPointPositions,
-		EPropertyChangeType::Type InChangeType, EAxisSnapMode InSnapMode
+		EPropertyChangeType::Type InChangeType, EKeyMoveMode InSnapMode
 		);
 	/** Generic function for updating control points. */
 	void OnSetControlPoints(
 		FPerCurveDeformer2D& InDeformer, FCurveModel& InCurveModel,
 		TConstArrayView<int32> InControlPointsToUpdate, TConstArrayView<FVector2D> InNewControlPoints,
-		EPropertyChangeType::Type InChangeType, EAxisSnapMode InSnapMode = EAxisSnapMode::Unrestricted
+		EPropertyChangeType::Type InChangeType, EKeyMoveMode InSnapMode = EKeyMoveMode::Unrestricted
 		);
 	
 	/** Moves the clicked edge to the opposite edge, effectively flattening the curve sections. */
diff --git a/Engine/Plugins/Editor/CurveEditorTools/Source/CurveEditorTools/Private/Tools/Lattice/Drag/EAxisSnapMode.h b/Engine/Plugins/Editor/CurveEditorTools/Source/CurveEditorTools/Private/Tools/Lattice/Drag/EAxisSnapMode.h
deleted file mode 100644
diff --git a/Engine/Plugins/Editor/CurveEditorTools/Source/CurveEditorTools/Private/Tools/Lattice/Drag/EKeyMoveMode.h b/Engine/Plugins/Editor/CurveEditorTools/Source/CurveEditorTools/Private/Tools/Lattice/Drag/EKeyMoveMode.h
new file mode 100644
--- /dev/null
+++ b/Engine/Plugins/Editor/CurveEditorTools/Source/CurveEditorTools/Private/Tools/Lattice/Drag/EKeyMoveMode.h
@@ -0,0 +1,24 @@
+﻿// Copyright Epic Games, Inc. All Rights Reserved.
+
+#pragma once
+
+#include "HAL/Platform.h"
+#include "Misc/EnumClassFlags.h"
+
+namespace UE::CurveEditorTools
+{
+/** Describes how keys are supposed to be snapped when using the lattice tool. */
+enum class EKeyMoveMode : uint8
+{
+	/** Mouse is unlocked, i.e can move freely in 2D space. */
+	Unrestricted,
+	
+	/** Do not move the keys in X direction. E.g. when user has shift pressed and is moving up / down (in y direction). */
+	IgnoreX,
+	/** Do not move the keys in Y direction. E.g. when user has shift pressed and is moving left / right (in x direction). */
+	IgnoreY,
+
+	IgnoreAll = IgnoreY | IgnoreX
+};
+ENUM_CLASS_FLAGS(EKeyMoveMode);
+}
\ No newline at end of file
diff --git a/Engine/Plugins/Editor/CurveEditorTools/Source/CurveEditorTools/Private/Tools/Lattice/Drag/LatticeDragOp.cpp b/Engine/Plugins/Editor/CurveEditorTools/Source/CurveEditorTools/Private/Tools/Lattice/Drag/LatticeDragOp.cpp
--- a/Engine/Plugins/Editor/CurveEditorTools/Source/CurveEditorTools/Private/Tools/Lattice/Drag/LatticeDragOp.cpp
+++ b/Engine/Plugins/Editor/CurveEditorTools/Source/CurveEditorTools/Private/Tools/Lattice/Drag/LatticeDragOp.cpp
@@ -47,29 +47,29 @@ void FLatticeDragOp::CancelDrag()
 
 namespace Private
 {
-static EAxisSnapMode DetermineSnapMode(const FCurveEditorAxisSnap& Snap, const FCurveEditorAxisSnap::FSnapState& SnapState)
+static EKeyMoveMode DetermineSnapMode(const FCurveEditorAxisSnap& Snap, const FCurveEditorAxisSnap::FSnapState& SnapState)
 {
 	const bool bIsRestrictingWithShift = SnapState.bHasPassedThreshold;
 	if (bIsRestrictingWithShift)
 	{
-		return FMath::IsNearlyZero(SnapState.MouseLockVector.X) ? EAxisSnapMode::IgnoreSnapOnX : EAxisSnapMode::IgnoreSnapOnY;
+		return FMath::IsNearlyZero(SnapState.MouseLockVector.X) ? EKeyMoveMode::IgnoreX : EKeyMoveMode::IgnoreY;
 	}
 
 	// This handles the start of the drag: shift is pressed but not yet above the threshold.
 	// Without this, all keys would get snapped and then become unsnapped once the mouse goes over the threshold. It's a visual unpleasantness.
 	if (SnapState.bHasStartPosition)
 	{
-		return EAxisSnapMode::IgnoreSnapping;
+		return EKeyMoveMode::IgnoreAll;
 	}
 	
 	switch (Snap.RestrictedAxisList)
 	{
-	case ECurveEditorSnapAxis::CESA_None: return EAxisSnapMode::Unrestricted;
-	case ECurveEditorSnapAxis::CESA_X: return EAxisSnapMode::IgnoreSnapOnY;
-	case ECurveEditorSnapAxis::CESA_Y: return EAxisSnapMode::IgnoreSnapOnX;
+	case ECurveEditorSnapAxis::CESA_None: return EKeyMoveMode::Unrestricted;
+	case ECurveEditorSnapAxis::CESA_X: return EKeyMoveMode::IgnoreY;
+	case ECurveEditorSnapAxis::CESA_Y: return EKeyMoveMode::IgnoreX;
 
 	default:
-		return EAxisSnapMode::Unrestricted;
+		return EKeyMoveMode::Unrestricted;
 	}
 }
 }
@@ -89,7 +89,7 @@ void FLatticeDragOp::AccumulateMouseMovement(const FGeometry& InGeometry, const
 	const FVector2D SnappedPosition = Snap.GetSnappedPosition(
 		InitialMousePosition, ScreenSpacePosition, LastMousePosition, InMouseEvent, SnapState
 		);
-	const EAxisSnapMode SnapMode = Private::DetermineSnapMode(Snap, SnapState);
+	const EKeyMoveMode SnapMode = Private::DetermineSnapMode(Snap, SnapState);
 	
 	if (!AccumulatedMouseMovement)
 	{
diff --git a/Engine/Plugins/Editor/CurveEditorTools/Source/CurveEditorTools/Private/Tools/Lattice/Drag/LatticeDragOp.h b/Engine/Plugins/Editor/CurveEditorTools/Source/CurveEditorTools/Private/Tools/Lattice/Drag/LatticeDragOp.h
--- a/Engine/Plugins/Editor/CurveEditorTools/Source/CurveEditorTools/Private/Tools/Lattice/Drag/LatticeDragOp.h
+++ b/Engine/Plugins/Editor/CurveEditorTools/Source/CurveEditorTools/Private/Tools/Lattice/Drag/LatticeDragOp.h
@@ -3,7 +3,7 @@
 #pragma once
 
 #include "CurveEditorSnapMetrics.h"
-#include "EAxisSnapMode.h"
+#include "EKeyMoveMode.h"
 #include "Layout/Geometry.h"
 #include "Math/MathFwd.h"
 #include "Math/Vector2D.h"
@@ -31,9 +31,9 @@ public:
 
 protected:
 
-	virtual void OnBeginDrag(const FGeometry& InGeometry, const FVector2D& InInitialMousePosition, EAxisSnapMode InSnapMode) {}
-	virtual void OnMoveMouse(const FGeometry& InGeometry, const FVector2D& InScreenPosition, EAxisSnapMode InSnapMode) {}
-	virtual void OnEndDrag(const FGeometry& InGeometry, const FVector2D& InMousePosition, EAxisSnapMode InSnapMode) {}
+	virtual void OnBeginDrag(const FGeometry& InGeometry, const FVector2D& InInitialMousePosition, EKeyMoveMode InSnapMode) {}
+	virtual void OnMoveMouse(const FGeometry& InGeometry, const FVector2D& InScreenPosition, EKeyMoveMode InSnapMode) {}
+	virtual void OnEndDrag(const FGeometry& InGeometry, const FVector2D& InMousePosition, EKeyMoveMode InSnapMode) {}
 	virtual void OnCancelDrag() {}
 
 	TWeakPtr<FCurveEditor> GetCurveEditor() const { return CurveEditor; }
@@ -59,9 +59,9 @@ private:
 		/** The mouse position that has been accumulated by processing all pumped input this tick. */
 		FVector2D AccumulatedPosition;
 		/** The mouse lock mode that has been accumulated by processing all pumped input this tick. */
-		EAxisSnapMode LastDragMode;
+		EKeyMoveMode LastDragMode;
 
-		explicit FAccumulatedMouseMovement(const FGeometry& InCachedGeometry, const FVector2D& InInitialPosition, EAxisSnapMode InInitialDragMode)
+		explicit FAccumulatedMouseMovement(const FGeometry& InCachedGeometry, const FVector2D& InInitialPosition, EKeyMoveMode InInitialDragMode)
 			: CachedGeometry(InCachedGeometry)
 			, AccumulatedPosition(InInitialPosition)
 			, LastDragMode(InInitialDragMode)
diff --git a/Engine/Plugins/Editor/CurveEditorTools/Source/CurveEditorTools/Private/Tools/Lattice/Drag/LatticeDragOp_MoveControlPoints.cpp b/Engine/Plugins/Editor/CurveEditorTools/Source/CurveEditorTools/Private/Tools/Lattice/Drag/LatticeDragOp_MoveControlPoints.cpp
--- a/Engine/Plugins/Editor/CurveEditorTools/Source/CurveEditorTools/Private/Tools/Lattice/Drag/LatticeDragOp_MoveControlPoints.cpp
+++ b/Engine/Plugins/Editor/CurveEditorTools/Source/CurveEditorTools/Private/Tools/Lattice/Drag/LatticeDragOp_MoveControlPoints.cpp
@@ -24,18 +24,18 @@ FLatticeDragOp_MoveControlPoints::FLatticeDragOp_MoveControlPoints(
 	, UpdatedControlPointPositions(InitialControlPoints)
 {}
 
-void FLatticeDragOp_MoveControlPoints::OnBeginDrag(const FGeometry& InGeometry, const FVector2D& InInitialMousePosition, EAxisSnapMode InSnapMode)
+void FLatticeDragOp_MoveControlPoints::OnBeginDrag(const FGeometry& InGeometry, const FVector2D& InInitialMousePosition, EKeyMoveMode InSnapMode)
 {
 	OnBeginDragControlPointsDelegate.ExecuteIfBound();
 }
 
-void FLatticeDragOp_MoveControlPoints::OnMoveMouse(const FGeometry& InGeometry, const FVector2D& InScreenPosition, EAxisSnapMode InSnapMode)
+void FLatticeDragOp_MoveControlPoints::OnMoveMouse(const FGeometry& InGeometry, const FVector2D& InScreenPosition, EKeyMoveMode InSnapMode)
 {
 	UpdateControlPoints(InGeometry, InScreenPosition);
 	OnInteractiveDragControlPointsDelegate.Execute(UpdatedControlPointPositions, InSnapMode);
 }
 
-void FLatticeDragOp_MoveControlPoints::OnEndDrag(const FGeometry& InGeometry, const FVector2D& InMousePosition, EAxisSnapMode InSnapMode)
+void FLatticeDragOp_MoveControlPoints::OnEndDrag(const FGeometry& InGeometry, const FVector2D& InMousePosition, EKeyMoveMode InSnapMode)
 {
 	UpdateControlPoints(InGeometry, InMousePosition);
 	OnFinishDragControlPointsDelegate.Execute(UpdatedControlPointPositions, InSnapMode);
diff --git a/Engine/Plugins/Editor/CurveEditorTools/Source/CurveEditorTools/Private/Tools/Lattice/Drag/LatticeDragOp_MoveControlPoints.h b/Engine/Plugins/Editor/CurveEditorTools/Source/CurveEditorTools/Private/Tools/Lattice/Drag/LatticeDragOp_MoveControlPoints.h
--- a/Engine/Plugins/Editor/CurveEditorTools/Source/CurveEditorTools/Private/Tools/Lattice/Drag/LatticeDragOp_MoveControlPoints.h
+++ b/Engine/Plugins/Editor/CurveEditorTools/Source/CurveEditorTools/Private/Tools/Lattice/Drag/LatticeDragOp_MoveControlPoints.h
@@ -22,7 +22,7 @@ class FLatticeDragOp_MoveControlPoints : public FLatticeDragOp
 public:
 
 	using FOnBeginDragOperation = FSimpleDelegate;
-	DECLARE_DELEGATE_TwoParams(FOnControlPointsMoved, const TConstArrayView<FVector2D>& InNewControlPointPositions, EAxisSnapMode InSnapMode);
+	DECLARE_DELEGATE_TwoParams(FOnControlPointsMoved, const TConstArrayView<FVector2D>& InNewControlPointPositions, EKeyMoveMode InSnapMode);
 
 	explicit FLatticeDragOp_MoveControlPoints(
 		TWeakPtr<FCurveEditor> InCurveEditor,
@@ -36,9 +36,9 @@ public:
 protected:
 	
 	//~ Begin FLatticeDragOp Interface
-	virtual void OnBeginDrag(const FGeometry& InGeometry, const FVector2D& InInitialMousePosition, EAxisSnapMode InSnapMode) override;
-	virtual void OnMoveMouse(const FGeometry& InGeometry, const FVector2D& InScreenPosition, EAxisSnapMode InSnapMode) override;
-	virtual void OnEndDrag(const FGeometry& InGeometry, const FVector2D& InMousePosition, EAxisSnapMode InSnapMode) override;
+	virtual void OnBeginDrag(const FGeometry& InGeometry, const FVector2D& InInitialMousePosition, EKeyMoveMode InSnapMode) override;
+	virtual void OnMoveMouse(const FGeometry& InGeometry, const FVector2D& InScreenPosition, EKeyMoveMode InSnapMode) override;
+	virtual void OnEndDrag(const FGeometry& InGeometry, const FVector2D& InMousePosition, EKeyMoveMode InSnapMode) override;
 	//~ End FLatticeDragOp Interface
 
 private:
diff --git a/Engine/Plugins/Editor/CurveEditorTools/Source/CurveEditorTools/Private/Tools/Lattice/LatticeFwd.h b/Engine/Plugins/Editor/CurveEditorTools/Source/CurveEditorTools/Private/Tools/Lattice/LatticeFwd.h
--- a/Engine/Plugins/Editor/CurveEditorTools/Source/CurveEditorTools/Private/Tools/Lattice/LatticeFwd.h
+++ b/Engine/Plugins/Editor/CurveEditorTools/Source/CurveEditorTools/Private/Tools/Lattice/LatticeFwd.h
@@ -7,9 +7,10 @@ struct FKeyHandle;
 namespace UE::CurveEditorTools
 {
 class FLatticeDeformer2D;
+struct FLatticePointMetaData;
 template<typename TPointMetaData> class TLatticeDeformer2D;
 
 /** We never add any points to this deformer. It's just used to keep track of the control points & edges displayed in the view. */
 using FGlobalLatticeDeformer2D = FLatticeDeformer2D;
-using FPerCurveDeformer2D = TLatticeDeformer2D<FKeyHandle>;
+using FPerCurveDeformer2D = TLatticeDeformer2D<FLatticePointMetaData>;
 }
diff --git a/Engine/Plugins/Editor/CurveEditorTools/Source/CurveEditorTools/Private/Tools/Lattice/LatticePointMetaData.h b/Engine/Plugins/Editor/CurveEditorTools/Source/CurveEditorTools/Private/Tools/Lattice/LatticePointMetaData.h
new file mode 100644
--- /dev/null
+++ b/Engine/Plugins/Editor/CurveEditorTools/Source/CurveEditorTools/Private/Tools/Lattice/LatticePointMetaData.h
@@ -0,0 +1,30 @@
+﻿// Copyright Epic Games, Inc. All Rights Reserved.
+
+#pragma once
+
+#include "CurveDataAbstraction.h"
+#include "Curves/KeyHandle.h"
+#include "Serialization/Archive.h"
+
+namespace UE::CurveEditorTools
+{
+/** Extra that stored for each point added to the lattice grid. */
+struct FLatticePointMetaData
+{
+	/** The handle for the key this lattice point corresponds to. */
+	FKeyHandle KeyHandle;
+	
+	/** The position the key had before the lattice started moving it. */
+	FKeyPosition OriginalPosition;
+	
+	friend FArchive& operator<<(FArchive& InArchive, FLatticePointMetaData& InMetaData);
+};
+
+inline FArchive& operator<<(FArchive& InArchive, FLatticePointMetaData& InMetaData)
+{
+	InArchive << InMetaData.KeyHandle;
+	InArchive << InMetaData.OriginalPosition.InputValue;
+	InArchive << InMetaData.OriginalPosition.OutputValue;
+	return InArchive;
+}
+}
diff --git a/Engine/Plugins/Editor/CurveEditorTools/Source/CurveEditorTools/Private/Tools/Lattice/Mirror/LatticeEdgeTangentMirrorOp.cpp b/Engine/Plugins/Editor/CurveEditorTools/Source/CurveEditorTools/Private/Tools/Lattice/Mirror/LatticeEdgeTangentMirrorOp.cpp
--- a/Engine/Plugins/Editor/CurveEditorTools/Source/CurveEditorTools/Private/Tools/Lattice/Mirror/LatticeEdgeTangentMirrorOp.cpp
+++ b/Engine/Plugins/Editor/CurveEditorTools/Source/CurveEditorTools/Private/Tools/Lattice/Mirror/LatticeEdgeTangentMirrorOp.cpp
@@ -5,8 +5,10 @@
 #include "Tools/Lattice/LatticeDeformer2D.h"
 #include "Tools/Lattice/Misc/LatticeUtils.h"
 #include "MirrorUtils.h"
+#include "Algo/AllOf.h"
 #include "Misc/Mirror/TangentMirrorSolver.h"
 #include "Misc/Optional.h"
+#include "Tools/Lattice/LatticePointMetaData.h"
 
 namespace UE::CurveEditorTools
 {
@@ -57,7 +59,7 @@ static bool GetDragStartAndMidpointHeights(
 }
 
 TOptional<FLatticeEdgeTangentsMirrorOp> FLatticeEdgeTangentsMirrorOp::MakeMirrorOpForDragLatticeEdge(
-	int32 InEdgeIndex, const FCurveModelID& InCurveId, const TLatticeDeformer2D<FKeyHandle>& InDeformer, const FCurveEditor& InCurveEditor
+	int32 InEdgeIndex, const FCurveModelID& InCurveId, const FPerCurveDeformer2D& InDeformer, const FCurveEditor& InCurveEditor
 	)
 {
 	if (!MirrorEdgeDetail::CanMirror(InEdgeIndex, InDeformer))
@@ -92,8 +94,9 @@ void FLatticeEdgeTangentsMirrorOp::OnMoveEdge(TConstArrayView<FVector2D> InNewEd
 }
 
 TOptional<FLatticeEdgeTangentsMirrorOp::FCurveMirrorData> FLatticeEdgeTangentsMirrorOp::ComputeTangentMirrorData(
-	int32 InEdgeIndex, const FCurveEditor& InCurveEditor,
-	const FLatticeDeformer2D& InCurveDeformer, const FCurveModelID& InCurve, TConstArrayView<FKeyHandle> InKeys
+	int32 InEdgeIndex,
+	const FCurveEditor& InCurveEditor, const FLatticeDeformer2D& InCurveDeformer, const FCurveModelID& InCurve,
+	TConstArrayView<FLatticePointMetaData> InMetaData
 	)
 {
 	double DragStartHeight, MidPointHeight;
@@ -102,11 +105,11 @@ TOptional<FLatticeEdgeTangentsMirrorOp::FCurveMirrorData> FLatticeEdgeTangentsMi
 		return {};
 	}
 	
+	TArray<FKeyHandle> KeyHandles;
+	Algo::Transform(InMetaData, KeyHandles, [](const FLatticePointMetaData& MetaData) { return MetaData.KeyHandle;} );
+	
 	CurveEditor::FTangentMirrorSolver TangentSolver(DragStartHeight, MidPointHeight);
-	if (TangentSolver.AddTangents(InCurveEditor, InCurve, InKeys))
-	{
-		return FCurveMirrorData(InCurve, MoveTemp(TangentSolver));
-	}
-	return {};
+	const bool bSuccess = TangentSolver.AddTangents(InCurveEditor, InCurve, KeyHandles);
+	return bSuccess ? FCurveMirrorData(InCurve, MoveTemp(TangentSolver)) : TOptional<FCurveMirrorData>();
 }
 }
diff --git a/Engine/Plugins/Editor/CurveEditorTools/Source/CurveEditorTools/Private/Tools/Lattice/Mirror/LatticeEdgeTangentMirrorOp.h b/Engine/Plugins/Editor/CurveEditorTools/Source/CurveEditorTools/Private/Tools/Lattice/Mirror/LatticeEdgeTangentMirrorOp.h
--- a/Engine/Plugins/Editor/CurveEditorTools/Source/CurveEditorTools/Private/Tools/Lattice/Mirror/LatticeEdgeTangentMirrorOp.h
+++ b/Engine/Plugins/Editor/CurveEditorTools/Source/CurveEditorTools/Private/Tools/Lattice/Mirror/LatticeEdgeTangentMirrorOp.h
@@ -2,10 +2,10 @@
 
 #pragma once
 
-#include "MirrorUtils.h"
 #include "HAL/Platform.h"
-#include "Tools/Lattice/Misc/LatticeUtils.h"
+#include "MirrorUtils.h"
 #include "Misc/Mirror/TangentMirrorSolver.h"
+#include "Tools/Lattice/LatticeFwd.h"
 
 class FCurveEditor;
 namespace UE::CurveEditor { class FTangentMirrorSolver; }
@@ -29,7 +29,7 @@ public:
 	 * @return Valid if InEdgeIndex is an edge that should do mirroring when dragged, i.e. a top or bottom edge
 	 */
 	static TOptional<FLatticeEdgeTangentsMirrorOp> MakeMirrorOpForDragLatticeEdge(
-		int32 InEdgeIndex, const FCurveModelID& InCurveId, const TLatticeDeformer2D<FKeyHandle>& InDeformer,
+		int32 InEdgeIndex, const FCurveModelID& InCurveId, const FPerCurveDeformer2D& InDeformer,
 		const FCurveEditor& InCurveEditor UE_LIFETIMEBOUND
 	);
 	
@@ -60,7 +60,7 @@ private:
 	/** Tracks keys that have tangents that can be mirrored. The mirror midpoints are computed based off of InCurveDeformer. */
 	static TOptional<FCurveMirrorData> ComputeTangentMirrorData(
 		int32 InEdgeIndex, const FCurveEditor& InCurveEditor,
-		const FLatticeDeformer2D& InCurveDeformer, const FCurveModelID& InCurve, TConstArrayView<FKeyHandle> InKeys
+		const FLatticeDeformer2D& InCurveDeformer, const FCurveModelID& InCurve, TConstArrayView<FLatticePointMetaData> InMetaData
 		);
 };
 }
diff --git a/Engine/Plugins/Editor/CurveEditorTools/Source/CurveEditorTools/Private/Tools/Lattice/Mirror/LatticePointTangentMirrorOp.cpp b/Engine/Plugins/Editor/CurveEditorTools/Source/CurveEditorTools/Private/Tools/Lattice/Mirror/LatticePointTangentMirrorOp.cpp
--- a/Engine/Plugins/Editor/CurveEditorTools/Source/CurveEditorTools/Private/Tools/Lattice/Mirror/LatticePointTangentMirrorOp.cpp
+++ b/Engine/Plugins/Editor/CurveEditorTools/Source/CurveEditorTools/Private/Tools/Lattice/Mirror/LatticePointTangentMirrorOp.cpp
@@ -2,9 +2,11 @@
 
 #include "LatticePointTangentMirrorOp.h"
 
+#include "Algo/AllOf.h"
 #include "MirrorUtils.h"
 #include "Misc/Optional.h"
 #include "Tools/Lattice/LatticeDeformer2D.h"
+#include "Tools/Lattice/LatticePointMetaData.h"
 #include "Tools/Lattice/Misc/LatticeUtils.h"
 
 namespace UE::CurveEditorTools
@@ -36,7 +38,7 @@ static void GetDragStartAndMidpointHeights(
 }
 
 TOptional<FLatticePointTangentsMirrorOp> FLatticePointTangentsMirrorOp::MakeMirrorOpForDragLatticeControlPoint(
-	int32 InControlPointIndex, const FCurveModelID& InCurveId, const TLatticeDeformer2D<FKeyHandle>& InDeformer, const FCurveEditor& InCurveEditor)
+	int32 InControlPointIndex, const FCurveModelID& InCurveId, const FPerCurveDeformer2D& InDeformer, const FCurveEditor& InCurveEditor)
 {
 	TOptional<FCurveMirrorData> CurveData = ComputeTangentMirrorData(InControlPointIndex, InCurveEditor, InCurveId, InDeformer);
 	if (CurveData)
@@ -70,7 +72,7 @@ void FLatticePointTangentsMirrorOp::OnMovePoint(const FVector2D& InControlPoint,
 }
 
 TOptional<FLatticePointTangentsMirrorOp::FCurveMirrorData> FLatticePointTangentsMirrorOp::ComputeTangentMirrorData(
-	int32 InControlPointIndex, const FCurveEditor& InCurveEditor, const FCurveModelID& InCurveId, const TLatticeDeformer2D<FKeyHandle>& InDeformer
+	int32 InControlPointIndex, const FCurveEditor& InCurveEditor, const FCurveModelID& InCurveId, const FPerCurveDeformer2D& InDeformer
 	)
 {
 	using namespace MirrorPointDetail;
@@ -114,12 +116,13 @@ TOptional<FLatticePointTangentsMirrorOp::FCurveMirrorData> FLatticePointTangents
 		}
 	};
 	
+	TArray<FKeyHandle> KeyHandles;
+	const TConstArrayView<FLatticePointMetaData> CellMetaData = InDeformer.GetCellMetaData(0);
+	Algo::Transform(CellMetaData, KeyHandles, [](const FLatticePointMetaData& MetaData) { return MetaData.KeyHandle;} );
+	
 	CurveEditor::FTangentMirrorSolver MirrorSolver(DragStartHeight, MidPointHeight);
-	const bool bHadTangentsToMirror = MirrorSolver.AddTangents(InCurveEditor, InCurveId, InDeformer.GetCellMetaData(0), ComputeFalloff);
-	if (bHadTangentsToMirror) // Only user-specified tangents need mirroring
-	{
-		return FCurveMirrorData(InCurveId, MoveTemp(MirrorSolver), MoveTemp(FalloffValues));
-	}
-	return {};
+	const bool bHadTangentsToMirror = MirrorSolver.AddTangents(InCurveEditor, InCurveId, KeyHandles, ComputeFalloff);
+	// Only user-specified tangents need mirroring
+	return bHadTangentsToMirror ? FCurveMirrorData(InCurveId, MoveTemp(MirrorSolver), MoveTemp(FalloffValues)) : TOptional<FCurveMirrorData>();
 }
 }
diff --git a/Engine/Plugins/Editor/CurveEditorTools/Source/CurveEditorTools/Private/Tools/Lattice/Mirror/LatticePointTangentMirrorOp.h b/Engine/Plugins/Editor/CurveEditorTools/Source/CurveEditorTools/Private/Tools/Lattice/Mirror/LatticePointTangentMirrorOp.h
--- a/Engine/Plugins/Editor/CurveEditorTools/Source/CurveEditorTools/Private/Tools/Lattice/Mirror/LatticePointTangentMirrorOp.h
+++ b/Engine/Plugins/Editor/CurveEditorTools/Source/CurveEditorTools/Private/Tools/Lattice/Mirror/LatticePointTangentMirrorOp.h
@@ -5,6 +5,7 @@
 #include "HAL/Platform.h"
 #include "Misc/Mirror/TangentMirrorSolver.h"
 #include "Tools/Lattice/LatticeDeformer2D.h"
+#include "Tools/Lattice/LatticeFwd.h"
 
 class FCurveEditor;
 namespace UE::CurveEditor { class FTangentMirrorSolver; }
@@ -33,7 +34,7 @@ public:
 	 * @return Set if there were keys to mirror when dragging the specifiec control point. Unset otherwise.
 	 */
 	static TOptional<FLatticePointTangentsMirrorOp> MakeMirrorOpForDragLatticeControlPoint(
-		int32 InControlPointIndex, const FCurveModelID& InCurveId, const TLatticeDeformer2D<FKeyHandle>& InDeformer,
+		int32 InControlPointIndex, const FCurveModelID& InCurveId, const FPerCurveDeformer2D& InDeformer,
 		const FCurveEditor& InCurveEditor UE_LIFETIMEBOUND
 		);
 	
@@ -67,7 +68,7 @@ private:
 	 * @return True, if any of the lattice's keys need mirroring.
 	 */
 	static TOptional<FCurveMirrorData> ComputeTangentMirrorData(
-		int32 InControlPointIndex, const FCurveEditor& CurveEditor, const FCurveModelID& InCurveId, const TLatticeDeformer2D<FKeyHandle>& InDeformer
+		int32 InControlPointIndex, const FCurveEditor& CurveEditor, const FCurveModelID& InCurveId, const FPerCurveDeformer2D& InDeformer
 		);
 };
 }
diff --git a/Engine/Plugins/Editor/CurveEditorTools/Source/CurveEditorTools/Private/Tools/Lattice/PerCurveLatticeData.cpp b/Engine/Plugins/Editor/CurveEditorTools/Source/CurveEditorTools/Private/Tools/Lattice/PerCurveLatticeData.cpp
--- a/Engine/Plugins/Editor/CurveEditorTools/Source/CurveEditorTools/Private/Tools/Lattice/PerCurveLatticeData.cpp
+++ b/Engine/Plugins/Editor/CurveEditorTools/Source/CurveEditorTools/Private/Tools/Lattice/PerCurveLatticeData.cpp
@@ -4,6 +4,7 @@
 
 #include "Containers/Map.h"
 #include "CurveEditor.h"
+#include "LatticePointMetaData.h"
 #include "SCurveEditorView.h"
 #include "Misc/LatticeDrawUtils.h"
 #include "Misc/LatticeUtils.h"
@@ -26,8 +27,8 @@ static void AddPointsToDeformer(TConstArrayView<FKeyHandle> InKeys, const FCurve
 		CurveModel->GetKeyPositions(SingleKeyView, SinglePositionView);
 
 		const FVector2D AsVector{ Position.InputValue, Position.OutputValue };
-		const TConstArrayView<FVector2D> KeyView(&AsVector, 1);
-		LocalDeformer.AddPoints_BeforeLatticeMoved({ KeyId }, KeyView);
+		const FLatticePointMetaData MetaData{ KeyId, Position };
+		LocalDeformer.AddPoints_BeforeLatticeMoved(MakeConstArrayView(&MetaData, 1), MakeConstArrayView(&AsVector, 1));
 	}
 }
 }
diff --git a/Engine/Plugins/Editor/CurveEditorTools/Source/CurveEditorTools/Private/Tools/Lattice/Undo/LatticeSnapshot.cpp b/Engine/Plugins/Editor/CurveEditorTools/Source/CurveEditorTools/Private/Tools/Lattice/Undo/LatticeSnapshot.cpp
--- a/Engine/Plugins/Editor/CurveEditorTools/Source/CurveEditorTools/Private/Tools/Lattice/Undo/LatticeSnapshot.cpp
+++ b/Engine/Plugins/Editor/CurveEditorTools/Source/CurveEditorTools/Private/Tools/Lattice/Undo/LatticeSnapshot.cpp
@@ -30,7 +30,7 @@ FArchive& operator<<(FArchive& InArchive, FPerCurveSnapshot& InSnapshot)
 {
 	InArchive << InSnapshot.GlobalToLocalTransform;
 	InArchive << InSnapshot.KeyCoords;
-	InArchive << InSnapshot.Keys;
+	InArchive << InSnapshot.MetaData;
 	return InArchive;
 }
 	
@@ -100,7 +100,7 @@ bool ApplySnapshot(
     		{ ControlPoints[TopLeft], ControlPoints[TopRight], ControlPoints[BottomRight], ControlPoints[BottomLeft] }
     		);
     	Deformer.SetKeyCoordsInCell(CellIndexZero, Snapshot.KeyCoords);
-    	Deformer.SetCellMetaData(CellIndexZero, Snapshot.Keys);
+    	Deformer.SetCellMetaData(CellIndexZero, Snapshot.MetaData);
     	
     	OutPerCurveData.Add(Pair.Key, MoveTemp(Data));
     }
@@ -128,7 +128,7 @@ FLatticeSnapshot TakeSnapshot(
 		{
 			.GlobalToLocalTransform = Pair.Value.GlobalDeformerToCurveDeformer,
 			.KeyCoords = TArray<FLatticeKeyCoords>(Pair.Value.CurveDeformer.GetKeyCoordsInCell(0)),
-			.Keys = TArray<FKeyHandle>(Pair.Value.CurveDeformer.GetCellMetaData(0))
+			.MetaData = TArray<FLatticePointMetaData>(Pair.Value.CurveDeformer.GetCellMetaData(0))
 		};
 		return TPair<FCurveModelID, FPerCurveSnapshot>(Pair.Key, MoveTemp(CurveSnapshot));
 	});
diff --git a/Engine/Plugins/Editor/CurveEditorTools/Source/CurveEditorTools/Private/Tools/Lattice/Undo/LatticeSnapshot.h b/Engine/Plugins/Editor/CurveEditorTools/Source/CurveEditorTools/Private/Tools/Lattice/Undo/LatticeSnapshot.h
--- a/Engine/Plugins/Editor/CurveEditorTools/Source/CurveEditorTools/Private/Tools/Lattice/Undo/LatticeSnapshot.h
+++ b/Engine/Plugins/Editor/CurveEditorTools/Source/CurveEditorTools/Private/Tools/Lattice/Undo/LatticeSnapshot.h
@@ -5,6 +5,7 @@
 #include "CurveEditorTypes.h"
 #include "Misc/Optional.h"
 #include "Tools/Lattice/LatticeFwd.h"
+#include "Tools/Lattice/LatticePointMetaData.h"
 #include "Tools/Lattice/PerCurveLatticeData.h"
 
 class FCurveEditor;
@@ -16,7 +17,8 @@ struct FPerCurveSnapshot
 {
 	FTransform2d GlobalToLocalTransform;
 	TArray<FLatticeKeyCoords> KeyCoords;
-	TArray<FKeyHandle> Keys;
+	/** The metadata for each point in the cell. */
+	TArray<FLatticePointMetaData> MetaData;
 	
 	friend FArchive& operator<<(FArchive& InArchive, FPerCurveSnapshot& InSnapshot);
 };
