Curve Editor: Fix non-unity compile issue.

#Synced-CL 51881412

diff --git a/Engine/Plugins/Editor/CurveEditorTools/Source/CurveEditorTools/Private/Tools/Lattice/PerCurveLatticeData.h b/Engine/Plugins/Editor/CurveEditorTools/Source/CurveEditorTools/Private/Tools/Lattice/PerCurveLatticeData.h
--- a/Engine/Plugins/Editor/CurveEditorTools/Source/CurveEditorTools/Private/Tools/Lattice/PerCurveLatticeData.h
+++ b/Engine/Plugins/Editor/CurveEditorTools/Source/CurveEditorTools/Private/Tools/Lattice/PerCurveLatticeData.h
@@ -6,6 +6,7 @@
 #include "LatticeDeformer2D.h"
 #include "LatticeFwd.h"
 #include "Math/TransformCalculus2D.h"
+#include "Tools/Lattice/LatticePointMetaData.h"
 
 class FCurveEditor;
 struct FCurveModelID;
